Target Platform API added to provide more information about variants and variant support.

diff --git a/Engine/Source/Developer/ProjectLauncher/Private/ProjectLauncherPrivatePCH.h b/Engine/Source/Developer/ProjectLauncher/Private/ProjectLauncherPrivatePCH.h
new file mode 100644
--- /dev/null
+++ b/Engine/Source/Developer/ProjectLauncher/Private/ProjectLauncherPrivatePCH.h
@@ -0,0 +1,103 @@
+// Copyright 1998-2014 Epic Games, Inc. All Rights Reserved.
+
+#pragma once
+
+#include "ProjectLauncher.h"
+
+
+/* Private dependencies
+ *****************************************************************************/
+
+#include "DesktopPlatformModule.h"
+#include "IFilter.h"
+#include "ModuleManager.h"
+#include "TargetDeviceServices.h"
+#include "TargetPlatform.h"
+#include "PlatformInfo.h"
+#include "TextFilter.h"
+
+
+/* Static helpers
+ *****************************************************************************/
+
+// @todo gmp: move this into a shared library or create a SImageButton widget.
+static TSharedRef<SButton> MakeImageButton( const FSlateBrush* ButtonImage, const FText& ButtonTip, const FOnClicked& OnClickedDelegate )
+{
+	return SNew(SButton)
+		.ButtonStyle( FEditorStyle::Get(), "NoBorder" )
+		.OnClicked(OnClickedDelegate)
+		.ToolTipText(ButtonTip)
+		.ContentPadding(2)
+		.VAlign(VAlign_Center)
+		.ForegroundColor(FSlateColor::UseForeground())
+		.Content()
+		[
+			SNew(SImage)
+				.Image(ButtonImage)
+				.ColorAndOpacity(FSlateColor::UseForeground())
+		];
+}
+
+
+/* Private includes
+ *****************************************************************************/
+
+#include "ProjectLauncherCommands.h"
+#include "ProjectLauncherModel.h"
+
+#include "SProjectLauncherValidation.h"
+
+#include "SProjectLauncherDelegates.h"
+#include "SProjectLauncherBuildConfigurationSelector.h"
+#include "SProjectLauncherCookModeSelector.h"
+#include "SProjectLauncherFormLabel.h"
+#include "SProjectLauncherInstanceTypeSelector.h"
+#include "SProjectLauncherProfileLaunchButton.h"
+#include "SProjectLauncherProfileNameDescEditor.h"
+#include "SProjectLauncherVariantSelector.h"
+
+#include "SProjectLauncherProjectPicker.h"
+#include "SProjectLauncherProjectPage.h"
+
+#include "SProjectLauncherBuildPage.h"
+
+#include "SProjectLauncherMapListRow.h"
+#include "SProjectLauncherCultureListRow.h"
+#include "SProjectLauncherPlatformListRow.h"
+#include "SProjectLauncherCookedPlatforms.h"
+#include "SProjectLauncherCookByTheBookSettings.h"
+#include "SProjectLauncherCookOnTheFlySettings.h"
+#include "SProjectLauncherCookPage.h"
+#include "SProjectLauncherSimpleCookPage.h"
+
+#include "SProjectLauncherDeviceGroupSelector.h"
+#include "SProjectLauncherSimpleDeviceListRow.h"
+#include "SProjectLauncherSimpleDeviceListView.h"
+#include "SProjectLauncherProfileListRow.h"
+#include "SProjectLauncherProfileListView.h"
+#include "SProjectLauncherDeployFileServerSettings.h"
+#include "SProjectLauncherDeployToDeviceSettings.h"
+#include "SProjectLauncherDeployRepositorySettings.h"
+#include "SProjectLauncherDeployPage.h"
+#include "SProjectLauncherDeployTargets.h"
+#include "SProjectLauncherDeployTargetListRow.h"
+
+#include "SProjectLauncherLaunchRoleEditor.h"
+#include "SProjectLauncherLaunchCustomRoles.h"
+#include "SProjectLauncherLaunchPage.h"
+
+#include "SProjectLauncherPackagingSettings.h"
+#include "SProjectLauncherPackagePage.h"
+
+#include "SProjectLauncherPreviewPage.h"
+
+#include "SProjectLauncherTaskListRow.h"
+#include "SProjectLauncherMessageListRow.h"
+#include "SProjectLauncherProgress.h"
+
+#include "SProjectLauncherSettings.h"
+#include "SProjectLauncherDeployTaskSettings.h"
+#include "SProjectLauncherBuildTaskSettings.h"
+#include "SProjectLauncherLaunchTaskSettings.h"
+
+#include "SProjectLauncher.h"
diff --git a/Engine/Source/Developer/ProjectLauncher/Private/Widgets/Build/SProjectLauncherBuildPage.cpp b/Engine/Source/Developer/ProjectLauncher/Private/Widgets/Build/SProjectLauncherBuildPage.cpp
new file mode 100644
--- /dev/null
+++ b/Engine/Source/Developer/ProjectLauncher/Private/Widgets/Build/SProjectLauncherBuildPage.cpp
@@ -0,0 +1,209 @@
+// Copyright 1998-2014 Epic Games, Inc. All Rights Reserved.
+
+/*=============================================================================
+	SProjectLauncherCookPage.cpp: Implements the SProjectLauncherCookPage class.
+=============================================================================*/
+
+#include "ProjectLauncherPrivatePCH.h"
+
+
+#define LOCTEXT_NAMESPACE "SProjectLauncherBuildPage"
+
+
+/* SProjectLauncherCookPage structors
+ *****************************************************************************/
+
+SProjectLauncherBuildPage::~SProjectLauncherBuildPage( )
+{
+	if (Model.IsValid())
+	{
+		Model->OnProfileSelected().RemoveAll(this);
+	}
+}
+
+
+/* SProjectLauncherCookPage interface
+ *****************************************************************************/
+
+void SProjectLauncherBuildPage::Construct( const FArguments& InArgs, const FProjectLauncherModelRef& InModel )
+{
+	Model = InModel;
+
+	ChildSlot
+	[
+		SNew(SVerticalBox)
+
+		+ SVerticalBox::Slot()
+			.AutoHeight()
+			[
+				SNew(SHorizontalBox)
+
+				+ SHorizontalBox::Slot()
+					.FillWidth(1.0)
+					.VAlign(VAlign_Center)
+					[
+						SNew(STextBlock)
+							.Text(LOCTEXT("BuildText", "Do you wish to build?"))
+					]
+
+				+ SHorizontalBox::Slot()
+					.AutoWidth()
+					.Padding(8.0, 0.0, 0.0, 0.0)
+					[
+						// build mode check box
+						SNew(SCheckBox)
+						.IsChecked(this, &SProjectLauncherBuildPage::HandleBuildIsChecked)
+						.OnCheckStateChanged(this, &SProjectLauncherBuildPage::HandleBuildCheckedStateChanged)
+					]
+			]
+
+        + SVerticalBox::Slot()
+            .AutoHeight()
+            .Padding(0.0f, 8.0f, 0.0f, 0.0f)
+            [
+                SNew(SExpandableArea)
+                .AreaTitle(LOCTEXT("AdvancedAreaTitle", "Advanced Settings"))
+                .InitiallyCollapsed(true)
+                .Padding(8.0)
+                .BodyContent()
+                [
+                    SNew(SVerticalBox)
+
+                    + SVerticalBox::Slot()
+                    .AutoHeight()
+                    [
+                        SNew(SButton)
+                        .Text(LOCTEXT("GenDSYMText", "Generate DSYM"))
+                        .IsEnabled( this, &SProjectLauncherBuildPage::HandleGenDSYMButtonEnabled )
+                        .OnClicked( this, &SProjectLauncherBuildPage::HandleGenDSYMClicked )
+                    ]
+                ]
+            ]
+
+		+ SVerticalBox::Slot()
+			.AutoHeight()
+			[
+				SNew(SProjectLauncherCookedPlatforms, InModel)
+					.Visibility(this, &SProjectLauncherBuildPage::HandleBuildPlatformVisibility)
+			]
+	];
+
+	Model->OnProfileSelected().AddSP(this, &SProjectLauncherBuildPage::HandleProfileManagerProfileSelected);
+}
+
+
+/* SProjectLauncherBuildPage implementation
+ *****************************************************************************/
+
+bool SProjectLauncherBuildPage::GenerateDSYMForProject( const FString& ProjectName, const FString& Configuration )
+{
+    // UAT executable
+    FString ExecutablePath = FPaths::ConvertRelativePathToFull(FPaths::EngineDir() + FString(TEXT("Build")) / TEXT("BatchFiles"));
+#if PLATFORM_MAC
+    FString Executable = TEXT("RunUAT.command");
+#elif PLATFORM_LINUX
+    FString Executable = TEXT("RunUAT.sh");
+#else
+    FString Executable = TEXT("RunUAT.bat");
+#endif
+
+    // build UAT command line parameters
+    FString CommandLine;
+    CommandLine = FString::Printf(TEXT("GenerateDSYM -project=%s -config=%s"),
+        *ProjectName,
+        *Configuration);
+
+    // launch the builder and monitor its process
+    FProcHandle ProcessHandle = FPlatformProcess::CreateProc(*(ExecutablePath / Executable), *CommandLine, false, false, false, NULL, 0, *ExecutablePath, NULL);
+
+	return ProcessHandle.Close();
+}
+
+
+/* SProjectLauncherBuildPage callbacks
+ *****************************************************************************/
+
+
+// Callback for changing the checked state of a platform menu check box. 
+void SProjectLauncherBuildPage::HandleBuildCheckedStateChanged( ESlateCheckBoxState::Type CheckState )
+{
+	ILauncherProfilePtr SelectedProfile = Model->GetSelectedProfile();
+
+	if (SelectedProfile.IsValid())
+	{
+		SelectedProfile->SetBuildGame(CheckState == ESlateCheckBoxState::Checked);
+	}
+}
+
+// Callback for determining whether a platform menu entry is checked.
+ESlateCheckBoxState::Type SProjectLauncherBuildPage::HandleBuildIsChecked() const
+{
+	ILauncherProfilePtr SelectedProfile = Model->GetSelectedProfile();
+
+	if (SelectedProfile.IsValid())
+	{
+		if (SelectedProfile->IsBuilding())
+		{
+			return ESlateCheckBoxState::Checked;
+		}
+	}
+
+	return ESlateCheckBoxState::Unchecked;
+}
+
+void SProjectLauncherBuildPage::HandleProfileManagerProfileSelected( const ILauncherProfilePtr& SelectedProfile, const ILauncherProfilePtr& PreviousProfile )
+{
+	// reload settings
+}
+
+EVisibility SProjectLauncherBuildPage::HandleBuildPlatformVisibility( ) const
+{
+	ILauncherProfilePtr SelectedProfile = Model->GetSelectedProfile();
+
+	if (SelectedProfile.IsValid())
+	{
+		if (SelectedProfile->GetCookMode() == ELauncherProfileCookModes::DoNotCook && SelectedProfile->GetDeploymentMode() == ELauncherProfileDeploymentModes::DoNotDeploy)
+		{
+			return EVisibility::Visible;
+		}
+	}
+
+	return EVisibility::Collapsed;
+}
+
+// Callback for pressing the Advanced Setting - Generate DSYM button.
+FReply SProjectLauncherBuildPage::HandleGenDSYMClicked()
+{
+    ILauncherProfilePtr SelectedProfile = Model->GetSelectedProfile();
+
+    if(SelectedProfile.IsValid())
+    {
+        if (!SelectedProfile->HasValidationError(ELauncherProfileValidationErrors::NoProjectSelected))
+        {
+            FString ProjectName = SelectedProfile->GetProjectName();
+            EBuildConfigurations::Type ProjectConfig = SelectedProfile->GetBuildConfiguration();
+
+            GenerateDSYMForProject( ProjectName, EBuildConfigurations::ToString(ProjectConfig) );
+        }
+    }
+
+    return FReply::Handled();
+}
+
+bool SProjectLauncherBuildPage::HandleGenDSYMButtonEnabled() const
+{
+    ILauncherProfilePtr SelectedProfile = Model->GetSelectedProfile();
+
+    if (SelectedProfile.IsValid())
+    {
+        if (!SelectedProfile->HasValidationError(ELauncherProfileValidationErrors::NoProjectSelected))
+        {
+            return true;
+        }
+    }
+
+    return false;
+}
+
+
+#undef LOCTEXT_NAMESPACE
diff --git a/Engine/Source/Developer/ProjectLauncher/Private/Widgets/Build/SProjectLauncherBuildPage.h b/Engine/Source/Developer/ProjectLauncher/Private/Widgets/Build/SProjectLauncherBuildPage.h
new file mode 100644
--- /dev/null
+++ b/Engine/Source/Developer/ProjectLauncher/Private/Widgets/Build/SProjectLauncherBuildPage.h
@@ -0,0 +1,69 @@
+// Copyright 1998-2014 Epic Games, Inc. All Rights Reserved.
+
+#pragma once
+
+
+/**
+ * Implements the profile page for the session launcher wizard.
+ */
+class SProjectLauncherBuildPage
+	: public SCompoundWidget
+{
+public:
+
+	SLATE_BEGIN_ARGS(SProjectLauncherBuildPage) { }
+	SLATE_END_ARGS()
+
+public:
+
+	/**
+	 * Destructor.
+	 */
+	~SProjectLauncherBuildPage( );
+
+public:
+
+	/**
+	 * Constructs the widget.
+	 *
+	 * @param InArgs The Slate argument list.
+	 * @param InModel The data model.
+	 */
+	void Construct(	const FArguments& InArgs, const FProjectLauncherModelRef& InModel );
+
+protected:
+
+	/**
+	 * Generates the DSYM for the project.
+	 *
+	 * @param ProjectName The name of the project.
+	 * @param Configuration The build configuration name.
+	 * @return true on success, false otherwise.
+	 */
+    bool GenerateDSYMForProject( const FString& ProjectName, const FString& Configuration );
+
+private:
+
+	// Callback for changing the checked state of a platform menu check box. 
+	void HandleBuildCheckedStateChanged( ESlateCheckBoxState::Type CheckState );
+
+	// Callback for determining whether a platform menu entry is checked.
+	ESlateCheckBoxState::Type HandleBuildIsChecked() const;
+
+	// Callback for changing the selected profile in the profile manager.
+	void HandleProfileManagerProfileSelected( const ILauncherProfilePtr& SelectedProfile, const ILauncherProfilePtr& PreviousProfile );
+
+	// Callback for determining if the build platform list should be displayed
+	EVisibility HandleBuildPlatformVisibility( ) const;
+
+    // Callback for pressing the Advanced Setting - Generate DSYM button.
+    FReply HandleGenDSYMClicked();
+
+	// Callback for getting the enabled state of the Generate DSYM button.
+    bool HandleGenDSYMButtonEnabled() const;
+
+private:
+
+	// Holds a pointer to the data model.
+	FProjectLauncherModelPtr Model;
+};
diff --git a/Engine/Source/Developer/ProjectLauncher/Private/Widgets/Cook/SProjectLauncherCookByTheBookSettings.h b/Engine/Source/Developer/ProjectLauncher/Private/Widgets/Cook/SProjectLauncherCookByTheBookSettings.h
new file mode 100644
--- /dev/null
+++ b/Engine/Source/Developer/ProjectLauncher/Private/Widgets/Cook/SProjectLauncherCookByTheBookSettings.h
@@ -0,0 +1,161 @@
+// Copyright 1998-2014 Epic Games, Inc. All Rights Reserved.
+
+/*=============================================================================
+	SProjectLauncherCookByTheBookSettings.h: Declares the SProjectLauncherCookByTheBookSettings class.
+=============================================================================*/
+
+#pragma once
+
+
+namespace EShowMapsChoices
+{
+	enum Type
+	{
+		/**
+		 * Show all available maps.
+		 */
+		ShowAllMaps,
+
+		/**
+		 * Only show maps that are to be cooked.
+		 */
+		ShowCookedMaps
+	};
+}
+
+
+/**
+ * Implements the cook-by-the-book settings panel.
+ */
+class SProjectLauncherCookByTheBookSettings
+	: public SCompoundWidget
+{
+public:
+
+	SLATE_BEGIN_ARGS(SProjectLauncherCookByTheBookSettings) { }
+	SLATE_END_ARGS()
+
+
+public:
+
+	/**
+	 * Destructor.
+	 */
+	~SProjectLauncherCookByTheBookSettings( );
+
+
+public:
+
+	/**
+	 * Constructs the widget.
+	 *
+	 * @param InArgs The Slate argument list.
+	 * @param InModel The data model.
+	 */
+	void Construct(	const FArguments& InArgs, const FProjectLauncherModelRef& InModel, bool InShowSimple = false );
+
+
+protected:
+
+	/**
+	 * Refreshes the list of available maps.
+	 */
+	void RefreshMapList( );
+
+	/**
+	 * Refreshes the list of available cultures.
+	 */
+	void RefreshCultureList( );
+
+
+private:
+
+	// Handles clicking the 'Select All Cultures' button.
+	void HandleAllCulturesHyperlinkNavigate( bool AllPlatforms );
+
+	// Handles determining the visibility of the 'Select All Cultures' button.
+	EVisibility HandleAllCulturesHyperlinkVisibility( ) const;
+
+	// Handles clicking the 'Select All Maps' button.
+	void HandleAllMapsHyperlinkNavigate( bool AllPlatforms );
+
+	// Handles selecting a build configuration for the cooker.
+	void HandleCookConfigurationSelectorConfigurationSelected( EBuildConfigurations::Type );
+
+	// Handles getting the content text of the cooker build configuration selector.
+	FString HandleCookConfigurationSelectorText( ) const;
+
+	// Handles check state changes of the 'Incremental' check box.
+	void HandleIncrementalCheckBoxCheckStateChanged( ESlateCheckBoxState::Type NewState );
+
+	// Handles determining the checked state of the 'Incremental' check box.
+	ESlateCheckBoxState::Type HandleIncrementalCheckBoxIsChecked( ) const;
+
+	// Handles generating a row widget in the map list view.
+	TSharedRef<ITableRow> HandleMapListViewGenerateRow( TSharedPtr<FString> InItem, const TSharedRef<STableViewBase>& OwnerTable );
+
+	// Handles generating a row widget in the culture list view.
+	TSharedRef<ITableRow> HandleCultureListViewGenerateRow( TSharedPtr<FString> InItem, const TSharedRef<STableViewBase>& OwnerTable );
+
+	// Handles determining the visibility of the 'Select All Maps' button.
+	EVisibility HandleMapSelectionHyperlinkVisibility( ) const;
+
+	// Handles getting the visibility of the map selection warning message.
+	EVisibility HandleNoMapSelectedBoxVisibility( ) const;
+
+	// Handles getting the text in the 'No maps' text block.
+	FText HandleNoMapsTextBlockText( ) const;
+
+	// Handles changing the selected profile in the profile manager.
+	void HandleProfileManagerProfileSelected( const ILauncherProfilePtr& SelectedProfile, const ILauncherProfilePtr& PreviousProfile );
+
+	// Handles checking the specified 'Show maps' check box.
+	void HandleShowCheckBoxCheckStateChanged( ESlateCheckBoxState::Type NewState, EShowMapsChoices::Type Choice );
+
+	// Handles determining the checked state of the specified 'Show maps' check box.
+	ESlateCheckBoxState::Type HandleShowCheckBoxIsChecked( EShowMapsChoices::Type Choice ) const;
+
+	// Handles checking the specified 'Unversioned' check box.
+	void HandleUnversionedCheckBoxCheckStateChanged( ESlateCheckBoxState::Type NewState );
+
+	// Handles determining the checked state of the specified 'Unversioned' check box.
+	ESlateCheckBoxState::Type HandleUnversionedCheckBoxIsChecked( ) const;
+
+	// Handles determining the visibility of a validation error icon.
+	EVisibility HandleValidationErrorIconVisibility( ELauncherProfileValidationErrors::Type Error ) const;
+
+	// Callback for getting the cookers additional options
+	FText HandleCookOptionsTextBlockText() const;
+
+	// Callback for changing the cookers additional options
+	void HandleCookerOptionsCommitted(const FText& NewText, ETextCommit::Type CommitType);
+
+	// Callback for updating any settings after the selected project has changed in the profile
+	void HandleProfileProjectChanged();
+
+	// creates the complex widget
+	TSharedRef<SWidget> MakeComplexWidget();
+
+	// creates the simple widget
+	TSharedRef<SWidget> MakeSimpleWidget();
+
+private:
+
+	// Holds the culture list.
+	TArray<TSharedPtr<FString> > CultureList;
+
+	// Holds the map list view.
+	TSharedPtr<SListView<TSharedPtr<FString> > > CultureListView;
+
+	// Holds the map list.
+	TArray<TSharedPtr<FString> > MapList;
+
+	// Holds the map list view.
+	TSharedPtr<SListView<TSharedPtr<FString> > > MapListView;
+
+	// Holds a pointer to the data model.
+	FProjectLauncherModelPtr Model;
+
+	// Holds the current 'Show maps' check box choice.
+	EShowMapsChoices::Type ShowMapsChoice;
+};
diff --git a/Engine/Source/Developer/ProjectLauncher/Private/Widgets/Cook/SProjectLauncherCookedPlatforms.h b/Engine/Source/Developer/ProjectLauncher/Private/Widgets/Cook/SProjectLauncherCookedPlatforms.h
new file mode 100644
--- /dev/null
+++ b/Engine/Source/Developer/ProjectLauncher/Private/Widgets/Cook/SProjectLauncherCookedPlatforms.h
@@ -0,0 +1,206 @@
+// Copyright 1998-2014 Epic Games, Inc. All Rights Reserved.
+
+#pragma once
+
+
+#define LOCTEXT_NAMESPACE "SProjectLauncherCookedPlatforms"
+
+
+/**
+ * Implements the cooked platforms panel.
+ */
+class SProjectLauncherCookedPlatforms
+	: public SCompoundWidget
+{
+public:
+
+	SLATE_BEGIN_ARGS(SProjectLauncherCookedPlatforms) { }
+	SLATE_END_ARGS()
+
+public:
+
+	/**
+	 * Constructs the widget.
+	 *
+	 * @param InArgs The Slate argument list.
+	 * @param InModel The data model.
+	 */
+	void Construct(	const FArguments& InArgs, const FProjectLauncherModelRef& InModel )
+	{
+		Model = InModel;
+
+		MakePlatformMenu();
+
+		ChildSlot
+		[
+			SNew(SVerticalBox)
+
+			+ SVerticalBox::Slot()
+				.FillHeight(1.0f)
+				[
+					// platform menu
+					SAssignNew(PlatformListView, SListView<TSharedPtr<FString> >)
+					.HeaderRow(
+					SNew(SHeaderRow)
+					.Visibility(EVisibility::Collapsed)
+
+					+ SHeaderRow::Column("PlatformName")
+					.DefaultLabel(LOCTEXT("PlatformListPlatformNameColumnHeader", "Platform"))
+					.FillWidth(1.0f)
+					)
+					.ItemHeight(16.0f)
+					.ListItemsSource(&PlatformList)
+					.OnGenerateRow(this, &SProjectLauncherCookedPlatforms::HandlePlatformListViewGenerateRow)
+					.SelectionMode(ESelectionMode::None)
+				]
+
+			+ SVerticalBox::Slot()
+				.AutoHeight()
+				.Padding(0.0f, 6.0f, 0.0f, 4.0f)
+				[
+					SNew(SSeparator)
+						.Orientation(Orient_Horizontal)
+				]
+
+			+ SVerticalBox::Slot()
+				.AutoHeight()
+				[
+					SNew(SHorizontalBox)
+
+					+ SHorizontalBox::Slot()
+						.FillWidth(1.0f)
+						.HAlign(HAlign_Right)
+						[
+							SNew(STextBlock)
+								.Text(LOCTEXT("SelectLabel", "Select:"))
+						]
+
+					+ SHorizontalBox::Slot()
+						.AutoWidth()
+						.Padding(8.0f, 0.0f)
+						[
+							// all platforms hyper-link
+							SNew(SHyperlink)
+								.OnNavigate(this, &SProjectLauncherCookedPlatforms::HandleAllPlatformsHyperlinkNavigate, true)
+								.Text(LOCTEXT("AllPlatformsHyperlinkLabel", "All"))
+								.ToolTipText(LOCTEXT("AllPlatformsButtonTooltip", "Select all available platforms."))
+								.Visibility(this, &SProjectLauncherCookedPlatforms::HandleAllPlatformsHyperlinkVisibility)									
+						]
+
+					+ SHorizontalBox::Slot()
+						.AutoWidth()
+						[
+							// no platforms hyper-link
+							SNew(SHyperlink)
+								.OnNavigate(this, &SProjectLauncherCookedPlatforms::HandleAllPlatformsHyperlinkNavigate, false)
+								.Text(LOCTEXT("NoPlatformsHyperlinkLabel", "None"))
+								.ToolTipText(LOCTEXT("NoPlatformsHyperlinkTooltip", "Deselect all platforms."))
+								.Visibility(this, &SProjectLauncherCookedPlatforms::HandleAllPlatformsHyperlinkVisibility)									
+						]
+				]
+		];
+	}
+
+protected:
+
+	/**
+	 * Builds the platform menu.
+	 *
+	 * @return Platform menu widget.
+	 */
+	void MakePlatformMenu( )
+	{
+		TArray<ITargetPlatform*> Platforms = GetTargetPlatformManager()->GetTargetPlatforms();
+
+		if (Platforms.Num() > 0)
+		{
+			PlatformList.Reset();
+			for (int32 PlatformIndex = 0; PlatformIndex < Platforms.Num(); ++PlatformIndex)
+			{
+				FString PlatformName = Platforms[PlatformIndex]->PlatformName();
+
+				PlatformList.Add(MakeShareable(new FString(PlatformName)));
+			}
+		}
+	}
+
+private:
+
+	// Callback for clicking the 'Select All Platforms' button.
+	void HandleAllPlatformsHyperlinkNavigate( bool AllPlatforms )
+	{
+		ILauncherProfilePtr SelectedProfile = Model->GetSelectedProfile();
+
+		if (SelectedProfile.IsValid())
+		{
+			if (AllPlatforms)
+			{
+				TArray<ITargetPlatform*> Platforms = GetTargetPlatformManager()->GetTargetPlatforms();
+
+				for (int32 PlatformIndex = 0; PlatformIndex < Platforms.Num(); ++PlatformIndex)
+				{
+					SelectedProfile->AddCookedPlatform(Platforms[PlatformIndex]->PlatformName());
+				}
+			}
+			else
+			{
+				SelectedProfile->ClearCookedPlatforms();
+			}
+		}
+	}
+
+	// Callback for determining the visibility of the 'Select All Platforms' button.
+	EVisibility HandleAllPlatformsHyperlinkVisibility( ) const
+	{
+		if (GetTargetPlatformManager()->GetTargetPlatforms().Num() > 1)
+		{
+			return EVisibility::Visible;
+		}
+
+		return EVisibility::Collapsed;
+	}
+
+	// Callback for getting the color of a platform menu check box.
+	FSlateColor HandlePlatformMenuEntryColorAndOpacity( FString PlatformName ) const
+	{
+		ILauncherProfilePtr SelectedProfile = Model->GetSelectedProfile();
+
+		if (SelectedProfile.IsValid())
+		{
+			ITargetPlatform* TargetPlatform = GetTargetPlatformManager()->FindTargetPlatform(PlatformName);
+
+			if (TargetPlatform != NULL)
+			{
+//				if (TargetPlatform->HasValidBuild(SelectedProfile->GetProjectPath(), SelectedProfile->GetBuildConfiguration()))
+				{
+					return FEditorStyle::GetColor("Foreground");
+				}
+			}
+		}
+
+		return FLinearColor::Yellow;
+	}
+
+	// Handles generating a row widget in the map list view.
+	TSharedRef<ITableRow> HandlePlatformListViewGenerateRow( TSharedPtr<FString> InItem, const TSharedRef<STableViewBase>& OwnerTable )
+	{
+		return SNew(SProjectLauncherPlatformListRow, Model.ToSharedRef())
+			.PlatformName(InItem)
+			.OwnerTableView(OwnerTable);
+	}
+
+private:
+
+	// Holds a pointer to the data model.
+	FProjectLauncherModelPtr Model;
+
+	// Holds the platform list.
+	TArray<TSharedPtr<FString> > PlatformList;
+
+	// Holds the platform list view.
+	TSharedPtr<SListView<TSharedPtr<FString> > > PlatformListView;
+
+};
+
+
+#undef LOCTEXT_NAMESPACE
diff --git a/Engine/Source/Developer/ProjectLauncher/Private/Widgets/Deploy/SProjectLauncherDeployFileServerSettings.cpp b/Engine/Source/Developer/ProjectLauncher/Private/Widgets/Deploy/SProjectLauncherDeployFileServerSettings.cpp
new file mode 100644
--- /dev/null
+++ b/Engine/Source/Developer/ProjectLauncher/Private/Widgets/Deploy/SProjectLauncherDeployFileServerSettings.cpp
@@ -0,0 +1,139 @@
+// Copyright 1998-2014 Epic Games, Inc. All Rights Reserved.
+
+#include "ProjectLauncherPrivatePCH.h"
+
+
+#define LOCTEXT_NAMESPACE "SProjectLauncherDeployFileServerSettings"
+
+
+/* SProjectLauncherDeployFileServerSettings interface
+ *****************************************************************************/
+
+void SProjectLauncherDeployFileServerSettings::Construct( const FArguments& InArgs, const FProjectLauncherModelRef& InModel )
+{
+	Model = InModel;
+
+	ChildSlot
+	[
+		SNew(SVerticalBox)
+
+		+ SVerticalBox::Slot()
+			.AutoHeight()
+			[
+				SNew(SBorder)
+					.Padding(8.0f)
+					.BorderImage(FEditorStyle::GetBrush("ToolPanel.GroupBorder"))
+					[
+						// deploy targets area
+						SNew(SProjectLauncherDeployTargets, InModel)
+					]
+			]
+
+		+ SVerticalBox::Slot()
+			.AutoHeight()
+			.Padding(0.0f, 8.0f, 0.0f, 0.0f)
+			[
+				SNew(SExpandableArea)
+					.AreaTitle(LOCTEXT("AdvancedAreaTitle", "Advanced Settings"))
+					.InitiallyCollapsed(true)
+					.Padding(8.0f)
+					.BodyContent()
+					[
+						SNew(SVerticalBox)
+
+						+ SVerticalBox::Slot()
+							.AutoHeight()
+							[
+								// hide file server check box
+								SNew(SCheckBox)
+									.IsChecked(this, &SProjectLauncherDeployFileServerSettings::HandleHideWindowCheckBoxIsChecked)
+									.OnCheckStateChanged(this, &SProjectLauncherDeployFileServerSettings::HandleHideWindowCheckBoxCheckStateChanged)
+									.Padding(FMargin(4.0f, 0.0f))
+									.ToolTipText(LOCTEXT("HideWindowCheckBoxTooltip", "If checked, the file server's console window will be hidden from your desktop."))
+									.Content()
+									[
+										SNew(STextBlock)
+											.Text(LOCTEXT("HideWindowCheckBoxText", "Hide the file server's console window"))
+									]
+							]
+
+						+ SVerticalBox::Slot()
+							.AutoHeight()
+							.Padding(0.0f, 4.0f, 0.0f, 0.0f)
+							[
+								// streaming file server check box
+								SNew(SCheckBox)
+									.IsChecked(this, &SProjectLauncherDeployFileServerSettings::HandleStreamingServerCheckBoxIsChecked)
+									.OnCheckStateChanged(this, &SProjectLauncherDeployFileServerSettings::HandleStreamingServerCheckBoxCheckStateChanged)
+									.Padding(FMargin(4.0f, 0.0f))
+									.ToolTipText(LOCTEXT("StreamingServerCheckBoxTooltip", "If checked, the file server uses an experimental implementation that can serve multiple files simultaneously."))
+									.Content()
+									[
+										SNew(STextBlock)
+											.Text(LOCTEXT("StreamingServerCheckBoxText", "Streaming server (experimental)"))
+									]
+							]
+					]
+			]
+	];
+}
+
+
+/* SProjectLauncherDeployFileServerSettings callbacks
+ *****************************************************************************/
+
+void SProjectLauncherDeployFileServerSettings::HandleHideWindowCheckBoxCheckStateChanged( ESlateCheckBoxState::Type NewState )
+{
+	ILauncherProfilePtr SelectedProfile = Model->GetSelectedProfile();
+
+	if (SelectedProfile.IsValid())
+	{
+		SelectedProfile->SetHideFileServerWindow(NewState == ESlateCheckBoxState::Checked);
+	}
+}
+
+
+ESlateCheckBoxState::Type SProjectLauncherDeployFileServerSettings::HandleHideWindowCheckBoxIsChecked( ) const
+{
+	ILauncherProfilePtr SelectedProfile = Model->GetSelectedProfile();
+
+	if (SelectedProfile.IsValid())
+	{
+		if (SelectedProfile->IsFileServerHidden())
+		{
+			return ESlateCheckBoxState::Checked;
+		}
+	}
+
+	return ESlateCheckBoxState::Unchecked;
+}
+
+
+void SProjectLauncherDeployFileServerSettings::HandleStreamingServerCheckBoxCheckStateChanged( ESlateCheckBoxState::Type NewState )
+{
+	ILauncherProfilePtr SelectedProfile = Model->GetSelectedProfile();
+
+	if (SelectedProfile.IsValid())
+	{
+		SelectedProfile->SetStreamingFileServer(NewState == ESlateCheckBoxState::Checked);
+	}
+}
+
+
+ESlateCheckBoxState::Type SProjectLauncherDeployFileServerSettings::HandleStreamingServerCheckBoxIsChecked( ) const
+{
+	ILauncherProfilePtr SelectedProfile = Model->GetSelectedProfile();
+
+	if (SelectedProfile.IsValid())
+	{
+		if (SelectedProfile->IsFileServerStreaming())
+		{
+			return ESlateCheckBoxState::Checked;
+		}
+	}
+
+	return ESlateCheckBoxState::Unchecked;
+}
+
+
+#undef LOCTEXT_NAMESPACE
diff --git a/Engine/Source/Developer/ProjectLauncher/Private/Widgets/Deploy/SProjectLauncherDeployFileServerSettings.h b/Engine/Source/Developer/ProjectLauncher/Private/Widgets/Deploy/SProjectLauncherDeployFileServerSettings.h
new file mode 100644
--- /dev/null
+++ b/Engine/Source/Developer/ProjectLauncher/Private/Widgets/Deploy/SProjectLauncherDeployFileServerSettings.h
@@ -0,0 +1,46 @@
+// Copyright 1998-2014 Epic Games, Inc. All Rights Reserved.
+
+#pragma once
+
+
+/**
+ * Implements the deploy-to-device settings panel.
+ */
+class SProjectLauncherDeployFileServerSettings
+	: public SCompoundWidget
+{
+public:
+
+	SLATE_BEGIN_ARGS(SProjectLauncherDeployFileServerSettings) { }
+	SLATE_END_ARGS()
+
+public:
+
+	/**
+	 * Constructs the widget.
+	 *
+	 * @param InArgs The Slate argument list.
+	 * @param InModel The data model.
+	 */
+	void Construct(	const FArguments& InArgs, const FProjectLauncherModelRef& InModel );
+
+
+private:
+
+	// Callback for check state changes of the 'Hide Window' check box.
+	void HandleHideWindowCheckBoxCheckStateChanged( ESlateCheckBoxState::Type NewState );
+
+	// Callback for determining the checked state of the 'Hide Window' check box.
+	ESlateCheckBoxState::Type HandleHideWindowCheckBoxIsChecked( ) const;
+
+	// Callback for check state changes of the 'Streaming Server' check box.
+	void HandleStreamingServerCheckBoxCheckStateChanged( ESlateCheckBoxState::Type NewState );
+
+	// Callback for determining the checked state of the 'Streaming Server' check box.
+	ESlateCheckBoxState::Type HandleStreamingServerCheckBoxIsChecked( ) const;
+
+private:
+
+	// Holds a pointer to the data model.
+	FProjectLauncherModelPtr Model;
+};
diff --git a/Engine/Source/Developer/ProjectLauncher/Private/Widgets/Deploy/SProjectLauncherDeployTargets.h b/Engine/Source/Developer/ProjectLauncher/Private/Widgets/Deploy/SProjectLauncherDeployTargets.h
new file mode 100644
--- /dev/null
+++ b/Engine/Source/Developer/ProjectLauncher/Private/Widgets/Deploy/SProjectLauncherDeployTargets.h
@@ -0,0 +1,89 @@
+// Copyright 1998-2014 Epic Games, Inc. All Rights Reserved.
+
+#pragma once
+
+/**
+ * Implements the deployment targets panel.
+ */
+class SProjectLauncherDeployTargets
+	: public SCompoundWidget
+{
+public:
+
+	SLATE_BEGIN_ARGS(SProjectLauncherDeployTargets) { }
+	SLATE_END_ARGS()
+
+public:
+
+	/**
+	 * Destructor.
+	 */
+	~SProjectLauncherDeployTargets( );
+
+public:
+
+	/**
+	 * Constructs the widget.
+	 *
+	 * @param InArgs The Slate argument list.
+	 * @param InModel The data model.
+	 */
+	void Construct(	const FArguments& InArgs, const FProjectLauncherModelRef& InModel);
+
+protected:
+
+	/**
+	 * Refreshes the list of device proxies.
+	 */
+	void RefreshDeviceProxyList( );
+
+private:
+
+	// Callback for selecting a different device group in the device group selector.
+	void HandleDeviceGroupSelectorGroupSelected( const ILauncherDeviceGroupPtr& SelectedGroup );
+
+	// Callback for getting the currently selected device group for a device list row.
+	ILauncherDeviceGroupPtr HandleDeviceListRowDeviceGroup( ) const;
+
+	// Callback for getting the enabled state of a device proxy list row.
+	bool HandleDeviceListRowIsEnabled( ITargetDeviceProxyPtr DeviceProxy ) const;
+
+	// Callback for getting the tool tip text of a device proxy list row.
+	FText HandleDeviceListRowToolTipText( ITargetDeviceProxyPtr DeviceProxy ) const;
+
+	// Callback for generating a row in the device list view.
+	TSharedRef<ITableRow> HandleDeviceProxyListViewGenerateRow( ITargetDeviceProxyPtr InItem, const TSharedRef<STableViewBase>& OwnerTable ) const;
+
+	// Callback for determining the visibility of the device list view.
+	EVisibility HandleDeviceProxyListViewVisibility( ) const;
+
+	// Callback for getting a flag indicating whether a valid device group is currently selected.
+	EVisibility HandleDeviceSelectorVisibility( ) const;
+
+	// Callback for getting the visibility of the 'No devices detected' text block.
+	EVisibility HandleNoDevicesBoxVisibility( ) const;
+
+	// Callback for getting the text in the 'no devices' text block.
+	FText HandleNoDevicesTextBlockText( ) const;
+
+	// Callback for determining the visibility of the 'Select a device group' text block.
+	EVisibility HandleSelectDeviceGroupWarningVisibility( ) const;
+
+	// Callback for when a device proxy has been added to the device proxy manager.
+	void HandleDeviceProxyManagerProxyAdded(const ITargetDeviceProxyRef& AddedProxy);
+
+	// Callback for when a device proxy has been removed from the device proxy manager.
+	void HandleDeviceProxyManagerProxyRemoved(const ITargetDeviceProxyRef& RemovedProxy);
+
+private:
+
+	// Holds the list of available device proxies.
+	TArray<ITargetDeviceProxyPtr> DeviceProxyList;
+
+	// Holds the device proxy list view .
+	TSharedPtr<SListView<ITargetDeviceProxyPtr> > DeviceProxyListView;
+
+	// Holds a pointer to the data model.
+	FProjectLauncherModelPtr Model;
+
+};
diff --git a/Engine/Source/Developer/ProjectLauncher/Private/Widgets/Deploy/SProjectLauncherDeployToDeviceSettings.cpp b/Engine/Source/Developer/ProjectLauncher/Private/Widgets/Deploy/SProjectLauncherDeployToDeviceSettings.cpp
new file mode 100644
--- /dev/null
+++ b/Engine/Source/Developer/ProjectLauncher/Private/Widgets/Deploy/SProjectLauncherDeployToDeviceSettings.cpp
@@ -0,0 +1,102 @@
+// Copyright 1998-2014 Epic Games, Inc. All Rights Reserved.
+
+#include "ProjectLauncherPrivatePCH.h"
+
+
+#define LOCTEXT_NAMESPACE "SProjectLauncherDeployToDeviceSettings"
+
+
+/* SProjectLauncherDeployToDeviceSettings interface
+ *****************************************************************************/
+
+void SProjectLauncherDeployToDeviceSettings::Construct( const FArguments& InArgs, const FProjectLauncherModelRef& InModel, EVisibility InShowAdvanced )
+{
+	Model = InModel;
+
+	ChildSlot
+	[
+		SNew(SVerticalBox)
+
+		+ SVerticalBox::Slot()
+			.AutoHeight()
+			[
+				SNew(SBorder)
+					.Padding(8.0f)
+					.BorderImage(FEditorStyle::GetBrush("ToolPanel.GroupBorder"))
+					[
+						// deploy targets area
+						SNew(SProjectLauncherDeployTargets, InModel)
+					]
+			]
+
+		+ SVerticalBox::Slot()
+			.AutoHeight()
+			.Padding(0.0f, 8.0f, 0.0f, 0.0f)
+			[
+				SNew(SVerticalBox)
+				.Visibility(InShowAdvanced)
+
+				+ SVerticalBox::Slot()
+					.AutoHeight()
+					[
+						SNew(SExpandableArea)
+							.AreaTitle(LOCTEXT("AdvancedAreaTitle", "Advanced Settings"))
+							.InitiallyCollapsed(true)
+							.Padding(8.0f)
+							.BodyContent()
+							[
+								SNew(SVerticalBox)
+
+								+ SVerticalBox::Slot()
+									.AutoHeight()
+									[
+										// unreal pak check box
+										SNew(SCheckBox)
+											.IsChecked(this, &SProjectLauncherDeployToDeviceSettings::HandleUnrealPakCheckBoxIsChecked)
+											.OnCheckStateChanged(this, &SProjectLauncherDeployToDeviceSettings::HandleUnrealPakCheckBoxCheckStateChanged)
+											.Padding(FMargin(4.0f, 0.0f))
+											.ToolTipText(LOCTEXT("UnrealPakCheckBoxTooltip", "If checked, the content will be deployed as a single UnrealPak file instead of many separate files."))
+											.Content()
+											[
+												SNew(STextBlock)
+													.Text(LOCTEXT("UnrealPakCheckBoxText", "Store all content in a single file (UnrealPak)"))
+											]
+									]
+							]
+					]
+			]
+	];
+}
+
+
+/* SProjectLauncherDeployToDeviceSettings callbacks
+ *****************************************************************************/
+
+void SProjectLauncherDeployToDeviceSettings::HandleUnrealPakCheckBoxCheckStateChanged( ESlateCheckBoxState::Type NewState )
+{
+	ILauncherProfilePtr SelectedProfile = Model->GetSelectedProfile();
+
+	if (SelectedProfile.IsValid())
+	{
+		SelectedProfile->SetDeployWithUnrealPak(NewState == ESlateCheckBoxState::Checked);
+	}
+}
+
+
+ESlateCheckBoxState::Type SProjectLauncherDeployToDeviceSettings::HandleUnrealPakCheckBoxIsChecked( ) const
+{
+	ILauncherProfilePtr SelectedProfile = Model->GetSelectedProfile();
+
+	if (SelectedProfile.IsValid())
+	{
+		if (SelectedProfile->IsPackingWithUnrealPak())
+		{
+			return ESlateCheckBoxState::Checked;
+		}
+	}
+
+	return ESlateCheckBoxState::Unchecked;
+}
+
+
+#undef LOCTEXT_NAMESPACE
diff --git a/Engine/Source/Developer/ProjectLauncher/Private/Widgets/Deploy/SProjectLauncherDeployToDeviceSettings.h b/Engine/Source/Developer/ProjectLauncher/Private/Widgets/Deploy/SProjectLauncherDeployToDeviceSettings.h
new file mode 100644
--- /dev/null
+++ b/Engine/Source/Developer/ProjectLauncher/Private/Widgets/Deploy/SProjectLauncherDeployToDeviceSettings.h
@@ -0,0 +1,39 @@
+// Copyright 1998-2014 Epic Games, Inc. All Rights Reserved.
+
+#pragma once
+
+
+/**
+ * Implements the deploy-to-device settings panel.
+ */
+class SProjectLauncherDeployToDeviceSettings
+	: public SCompoundWidget
+{
+public:
+
+	SLATE_BEGIN_ARGS(SProjectLauncherDeployToDeviceSettings) { }
+	SLATE_END_ARGS()
+
+public:
+
+	/**
+	 * Constructs the widget.
+	 *
+	 * @param InArgs - The Slate argument list.
+	 * @param InModel - The data model.
+	 */
+	void Construct(	const FArguments& InArgs, const FProjectLauncherModelRef& InModel, EVisibility InShowAdvanced = EVisibility::Visible );
+
+private:
+
+	// Callback for check state changes of the 'UnrealPak' check box.
+	void HandleUnrealPakCheckBoxCheckStateChanged( ESlateCheckBoxState::Type NewState );
+
+	// Callback for determining the checked state of the 'UnrealPak' check box.
+	ESlateCheckBoxState::Type HandleUnrealPakCheckBoxIsChecked( ) const;
+
+private:
+
+	// Holds a pointer to the data model.
+	FProjectLauncherModelPtr Model;
+};
diff --git a/Engine/Source/Developer/ProjectLauncher/Private/Widgets/Deploy/SProjectLauncherDeviceGroupSelector.h b/Engine/Source/Developer/ProjectLauncher/Private/Widgets/Deploy/SProjectLauncherDeviceGroupSelector.h
new file mode 100644
--- /dev/null
+++ b/Engine/Source/Developer/ProjectLauncher/Private/Widgets/Deploy/SProjectLauncherDeviceGroupSelector.h
@@ -0,0 +1,244 @@
+// Copyright 1998-2014 Epic Games, Inc. All Rights Reserved.
+
+#pragma once
+
+
+#define LOCTEXT_NAMESPACE "SProjectLauncherDeviceGroupSelector"
+
+
+/**
+ * Delegate type for device group selection changes.
+ *
+ * The first parameter is the selected device group (or NULL if the previously selected group was unselected).
+ */
+DECLARE_DELEGATE_OneParam(FOnProjectLauncherDeviceGroupSelected, const ILauncherDeviceGroupPtr&)
+
+
+/**
+ * Implements a widget for device group selection.
+ */
+class SProjectLauncherDeviceGroupSelector
+	: public SCompoundWidget
+{
+public:
+
+	SLATE_BEGIN_ARGS(SProjectLauncherDeviceGroupSelector) { }
+
+		/**
+		 * Exposes the initially selected device group.
+		 */
+		SLATE_ARGUMENT(ILauncherDeviceGroupPtr, InitiallySelectedGroup)
+		
+		/**
+		 * Exposes a delegate to be invoked when a different device group has been selected.
+		 */
+		SLATE_EVENT(FOnProjectLauncherDeviceGroupSelected, OnGroupSelected)
+
+	SLATE_END_ARGS()
+
+public:
+
+	/**
+	 * Destructor.
+	 */
+	~SProjectLauncherDeviceGroupSelector( )
+	{
+		if (ProfileManager.IsValid())
+		{
+			ProfileManager->OnDeviceGroupAdded().RemoveAll(this);
+			ProfileManager->OnDeviceGroupRemoved().RemoveAll(this);
+		}
+	}
+
+public:
+
+	/**
+	 * Constructs the widget.
+	 *
+	 * @param InArgs - The Slate argument list.
+	 * @param InProfileManager - The profile manager to use.
+	 */
+	void Construct( const FArguments& InArgs, const ILauncherProfileManagerRef& InProfileManager )
+	{
+		OnGroupSelected = InArgs._OnGroupSelected;
+
+		ProfileManager = InProfileManager;
+
+		ChildSlot
+		[
+			SNew(SVerticalBox)
+
+			+ SVerticalBox::Slot()
+				.AutoHeight()
+				[
+					SNew(SProjectLauncherFormLabel)
+						.LabelText(LOCTEXT("DeviceGroupComboBoxLabel", "Device group to deploy to:"))
+				]
+
+			+ SVerticalBox::Slot()
+				.AutoHeight()
+				.Padding(0.0, 4.0, 0.0, 0.0)
+				[
+					SAssignNew(DeviceGroupComboBox, SEditableComboBox<ILauncherDeviceGroupPtr>)
+						.InitiallySelectedItem(InArgs._InitiallySelectedGroup)
+						.OptionsSource(&ProfileManager->GetAllDeviceGroups())
+						.AddButtonToolTip(LOCTEXT("AddProfileButtonToolTip", "Add a new device group"))
+						.RemoveButtonToolTip(LOCTEXT("DeleteProfileButtonToolTip", "Delete the selected device group"))
+						.RenameButtonToolTip(LOCTEXT("RenameProfileButtonToolTip", "Rename the selected device group"))
+						.OnAddClicked(this, &SProjectLauncherDeviceGroupSelector::HandleDeviceGroupComboBoxAddClicked)
+						.OnGenerateWidget(this, &SProjectLauncherDeviceGroupSelector::HandleDeviceGroupComboBoxGenerateWidget)
+						.OnGetEditableText(this, &SProjectLauncherDeviceGroupSelector::HandleDeviceGroupComboBoxGetEditableText)
+						.OnRemoveClicked(this, &SProjectLauncherDeviceGroupSelector::HandleDeviceGroupComboBoxRemoveClicked)
+						.OnSelectionChanged(this, &SProjectLauncherDeviceGroupSelector::HandleDeviceGroupComboBoxSelectionChanged)
+						.OnSelectionRenamed(this, &SProjectLauncherDeviceGroupSelector::HandleDeviceGroupComboBoxSelectionRenamed)
+						.Content()
+						[
+							SNew(STextBlock)
+								.Text(this, &SProjectLauncherDeviceGroupSelector::HandleDeviceGroupComboBoxContent)
+						]
+				]
+		];
+
+		ProfileManager->OnDeviceGroupAdded().AddSP(this, &SProjectLauncherDeviceGroupSelector::HandleProfileManagerDeviceGroupsChanged);
+		ProfileManager->OnDeviceGroupRemoved().AddSP(this, &SProjectLauncherDeviceGroupSelector::HandleProfileManagerDeviceGroupsChanged);
+	}
+
+	/**
+	 * Gets the currently selected device group.
+	 *
+	 * @return The selected group, or NULL if no group is selected.
+	 *
+	 * @see SetSelectedGroup
+	 */
+	ILauncherDeviceGroupPtr GetSelectedGroup( ) const
+	{
+		return DeviceGroupComboBox->GetSelectedItem();
+	}
+
+	/**
+	 * Sets the selected device group.
+	 *
+	 * @param DeviceGroup - The device group to select.
+	 *
+	 * @see GetSelectedGroup
+	 */
+	void SetSelectedGroup( const ILauncherDeviceGroupPtr& DeviceGroup )
+	{
+		if (!DeviceGroup.IsValid() || ProfileManager->GetAllDeviceGroups().Contains(DeviceGroup))
+		{
+			DeviceGroupComboBox->SetSelectedItem(DeviceGroup);
+		}		
+	}
+
+private:
+
+	// Callback for clicking the 'Add device group' button
+	FReply HandleDeviceGroupComboBoxAddClicked( )
+	{
+		ILauncherDeviceGroupPtr NewGroup = ProfileManager->AddNewDeviceGroup();
+
+		DeviceGroupComboBox->SetSelectedItem(NewGroup);
+
+		return FReply::Handled();
+	}
+
+	// Callback for getting the device group combo box content.
+	FString HandleDeviceGroupComboBoxContent() const
+	{
+		ILauncherDeviceGroupPtr SelectedGroup = DeviceGroupComboBox->GetSelectedItem();
+
+		if (SelectedGroup.IsValid())
+		{
+			return SelectedGroup->GetName();
+		}
+
+		return LOCTEXT("CreateOrSelectGroupText", "Create or select a device group...").ToString();
+	}
+
+	// Callback for generating a row for the device group combo box.
+	TSharedRef<SWidget> HandleDeviceGroupComboBoxGenerateWidget( ILauncherDeviceGroupPtr InItem )
+	{
+		return SNew(STextBlock)
+			.Text(this, &SProjectLauncherDeviceGroupSelector::HandleDeviceGroupComboBoxWidgetText, InItem);
+	}
+
+	// Callback for getting the editable text of the currently selected device group.
+	FString HandleDeviceGroupComboBoxGetEditableText( )
+	{
+		ILauncherDeviceGroupPtr SelectedGroup = DeviceGroupComboBox->GetSelectedItem();
+
+		if (SelectedGroup.IsValid())
+		{
+			return SelectedGroup->GetName();
+		}
+
+		return FString();
+	}
+
+	// Callback for clicking the 'Delete device group' button
+	FReply HandleDeviceGroupComboBoxRemoveClicked( )
+	{
+		ILauncherDeviceGroupPtr SelectedGroup = DeviceGroupComboBox->GetSelectedItem();
+
+		if (SelectedGroup.IsValid())
+		{
+			ProfileManager->RemoveDeviceGroup(SelectedGroup.ToSharedRef());
+		}
+
+		if (ProfileManager->GetAllDeviceGroups().Num() > 0)
+		{
+			DeviceGroupComboBox->SetSelectedItem(ProfileManager->GetAllDeviceGroups()[0]);
+		}
+		else
+		{
+			DeviceGroupComboBox->SetSelectedItem(NULL);
+		}
+
+		return FReply::Handled();
+	}
+
+	// Callback for changing the selected device group in the device group combo box.
+	void HandleDeviceGroupComboBoxSelectionChanged( ILauncherDeviceGroupPtr Selection, ESelectInfo::Type SelectInfo )
+	{
+		OnGroupSelected.ExecuteIfBound(Selection);
+	}
+
+	// Callback for when the selected item in the device group combo box has been renamed.
+	void HandleDeviceGroupComboBoxSelectionRenamed( const FText& CommittedText, ETextCommit::Type )
+	{
+		DeviceGroupComboBox->GetSelectedItem()->SetName(CommittedText.ToString());
+	}
+
+	// Callback for getting the text of a device group combo box widget.
+	FString HandleDeviceGroupComboBoxWidgetText( ILauncherDeviceGroupPtr Group ) const
+	{
+		if (Group.IsValid())
+		{
+			return Group->GetName();
+		}
+
+		return FString();
+	}
+
+	// Callback for changing the list of groups in the profile manager.
+	void HandleProfileManagerDeviceGroupsChanged( const ILauncherDeviceGroupRef& /*ChangedProfile*/ )
+	{
+		DeviceGroupComboBox->RefreshOptions();
+	}
+
+private:
+
+	// Holds the device group combo box.
+	TSharedPtr<SEditableComboBox<ILauncherDeviceGroupPtr> > DeviceGroupComboBox;
+
+	// Holds the profile manager.
+	ILauncherProfileManagerPtr ProfileManager;
+
+private:
+
+	// Holds a delegate to be invoked when a different device group has been selected.
+	FOnProjectLauncherDeviceGroupSelected OnGroupSelected;
+};
+
+
+#undef LOCTEXT_NAMESPACE
diff --git a/Engine/Source/Developer/ProjectLauncher/Private/Widgets/Deploy/SProjectLauncherSimpleDeviceListView.cpp b/Engine/Source/Developer/ProjectLauncher/Private/Widgets/Deploy/SProjectLauncherSimpleDeviceListView.cpp
new file mode 100644
--- /dev/null
+++ b/Engine/Source/Developer/ProjectLauncher/Private/Widgets/Deploy/SProjectLauncherSimpleDeviceListView.cpp
@@ -0,0 +1,148 @@
+// Copyright 1998-2014 Epic Games, Inc. All Rights Reserved.
+
+#include "ProjectLauncherPrivatePCH.h"
+
+
+#define LOCTEXT_NAMESPACE "SProjectLauncherSimpleDeviceListView"
+
+
+/* SProjectLauncherDeployTargets structors
+ *****************************************************************************/
+
+SProjectLauncherSimpleDeviceListView::~SProjectLauncherSimpleDeviceListView( )
+{
+	if (Model.IsValid())
+	{
+		const ITargetDeviceProxyManagerRef& DeviceProxyManager = Model->GetDeviceProxyManager();
+		DeviceProxyManager->OnProxyAdded().RemoveAll(this);
+		DeviceProxyManager->OnProxyRemoved().RemoveAll(this);
+	}
+}
+
+
+/* SProjectLauncherDeployTargets interface
+ *****************************************************************************/
+
+BEGIN_SLATE_FUNCTION_BUILD_OPTIMIZATION
+void SProjectLauncherSimpleDeviceListView::Construct(const FArguments& InArgs, const FProjectLauncherModelRef& InModel)
+{
+	OnProfileRun = InArgs._OnProfileRun;
+	IsAdvanced = InArgs._IsAdvanced;
+	
+	Model = InModel;
+
+	SAssignNew(DeviceProxyListView, SListView<ITargetDeviceProxyPtr>)
+		.SelectionMode(ESelectionMode::None)
+		.ListItemsSource(&DeviceProxyList)
+		.OnGenerateRow(this, &SProjectLauncherSimpleDeviceListView::HandleDeviceProxyListViewGenerateRow)
+		.ItemHeight(16.0f);
+
+	ChildSlot
+	[
+		SNew(SVerticalBox)
+
+		+ SVerticalBox::Slot()
+		.FillHeight(1.0f)
+		[
+			SNew(SScrollBorder, DeviceProxyListView.ToSharedRef())
+			[
+				DeviceProxyListView.ToSharedRef()
+			]
+		]
+
+		+ SVerticalBox::Slot()
+		.AutoHeight()
+		.Padding(2, 4, 2, 4)
+		[
+			SNew(SHorizontalBox)
+
+			+ SHorizontalBox::Slot()
+			.AutoWidth()
+			.HAlign(HAlign_Left)
+			[
+				SNew(STextBlock)
+				.TextStyle(FCoreStyle::Get(), "ToolBar.Keybinding")
+				.Text(LOCTEXT("ProjectLauncherDeviceManagerLinkPreamble", "Don't see your device? Verify it's setup and claimed in the "))
+			]
+
+			+ SHorizontalBox::Slot()
+			.AutoWidth()
+			.HAlign(HAlign_Left)
+			.Padding(0.0f, 0.0f)
+			[
+				// all cultures hyper link
+				SNew(SHyperlink)
+				.TextStyle(FCoreStyle::Get(), "ToolBar.Keybinding")
+				.OnNavigate(this, &SProjectLauncherSimpleDeviceListView::HandleDeviceManagerHyperlinkNavigate)
+				.Text(LOCTEXT("ProjectLauncherDeviceManagerLink", "Device Manager."))
+				.ToolTipText(LOCTEXT("ProjectLauncherDeviceManagerLinkTooltip", "Open the Device Manager window, where you can setup and claim devices connected to your machine or shared on the network."))
+			]
+		]
+	];
+
+	const ITargetDeviceProxyManagerRef& DeviceProxyManager = Model->GetDeviceProxyManager();
+
+	DeviceProxyManager->OnProxyAdded().AddSP(this, &SProjectLauncherSimpleDeviceListView::HandleDeviceProxyManagerProxyAdded);
+	DeviceProxyManager->OnProxyRemoved().AddSP(this, &SProjectLauncherSimpleDeviceListView::HandleDeviceProxyManagerProxyRemoved);
+
+	DeviceProxyManager->GetProxies(NAME_None, false, DeviceProxyList);
+
+}
+END_SLATE_FUNCTION_BUILD_OPTIMIZATION
+
+
+/* SProjectLauncherDeployTargets implementation
+ *****************************************************************************/
+
+void SProjectLauncherSimpleDeviceListView::RefreshDeviceProxyList()
+{
+	Model->GetDeviceProxyManager()->GetProxies(NAME_None, false, DeviceProxyList);
+	DeviceProxyListView->RequestListRefresh();
+}
+
+
+/* SProjectLauncherDeployTargets callbacks
+ *****************************************************************************/
+
+bool SProjectLauncherSimpleDeviceListView::HandleDeviceListRowIsEnabled( ITargetDeviceProxyPtr DeviceProxy ) const
+{
+	return true;
+}
+
+void SProjectLauncherSimpleDeviceListView::HandleDeviceManagerHyperlinkNavigate() const
+{
+	FGlobalTabmanager::Get()->InvokeTab(FTabId("DeviceManager"));
+}
+
+FText SProjectLauncherSimpleDeviceListView::HandleDeviceListRowToolTipText( ITargetDeviceProxyPtr DeviceProxy ) const
+{
+	FTextBuilder Builder;
+	Builder.AppendLineFormat(LOCTEXT("DeviceListRowToolTipName", "Name: {0}"), FText::FromString(DeviceProxy->GetName()));
+	//Builder.AppendLineFormat(LOCTEXT("DeviceListRowToolTipPlatform", "Platform: {0}"), FText::FromString(DeviceProxy->GetTargetPlatformName(SimpleProfile->GetDeviceVariant())));
+	//Builder.AppendLineFormat(LOCTEXT("DeviceListRowToolTipDeviceId", "Device ID: {0}"), FText::FromString(DeviceProxy->GetTargetDeviceId(SimpleProfile->GetDeviceVariant())));
+
+	return Builder.ToText();
+}
+
+
+TSharedRef<ITableRow> SProjectLauncherSimpleDeviceListView::HandleDeviceProxyListViewGenerateRow(ITargetDeviceProxyPtr InItem, const TSharedRef<STableViewBase>& OwnerTable) const
+{
+	return SNew(SProjectLauncherSimpleDeviceListRow, Model.ToSharedRef(), OwnerTable)
+		.OnProfileRun(OnProfileRun)
+		.IsAdvanced(IsAdvanced)
+		.DeviceProxy(InItem)
+		.IsEnabled(this, &SProjectLauncherSimpleDeviceListView::HandleDeviceListRowIsEnabled, InItem)
+		.ToolTipText(this, &SProjectLauncherSimpleDeviceListView::HandleDeviceListRowToolTipText, InItem);
+}
+
+void SProjectLauncherSimpleDeviceListView::HandleDeviceProxyManagerProxyAdded(const ITargetDeviceProxyRef& AddedProxy)
+{
+	RefreshDeviceProxyList();
+}
+
+void SProjectLauncherSimpleDeviceListView::HandleDeviceProxyManagerProxyRemoved(const ITargetDeviceProxyRef& AddedProxy)
+{
+	RefreshDeviceProxyList();
+}
+
+#undef LOCTEXT_NAMESPACE
diff --git a/Engine/Source/Developer/ProjectLauncher/Private/Widgets/Deploy/SProjectLauncherSimpleDeviceListView.h b/Engine/Source/Developer/ProjectLauncher/Private/Widgets/Deploy/SProjectLauncherSimpleDeviceListView.h
new file mode 100644
--- /dev/null
+++ b/Engine/Source/Developer/ProjectLauncher/Private/Widgets/Deploy/SProjectLauncherSimpleDeviceListView.h
@@ -0,0 +1,82 @@
+// Copyright 1998-2014 Epic Games, Inc. All Rights Reserved.
+
+#pragma once
+
+/**
+ * Implements the deployment targets panel.
+ */
+class SProjectLauncherSimpleDeviceListView
+	: public SCompoundWidget
+{
+public:
+
+	SLATE_BEGIN_ARGS(SProjectLauncherSimpleDeviceListView) { }
+		SLATE_EVENT(FOnProfileRun, OnProfileRun)
+		SLATE_ATTRIBUTE(bool, IsAdvanced)
+	SLATE_END_ARGS()
+
+public:
+
+	/**
+	 * Destructor.
+	 */
+	~SProjectLauncherSimpleDeviceListView( );
+
+public:
+
+	/**
+	 * Constructs the widget.
+	 *
+	 * @param InArgs The Slate argument list.
+	 * @param InModel The data model.
+	 * @param InAdvanced Whether or not the elements should show advanced data.
+	 */
+	void Construct(const FArguments& InArgs, const FProjectLauncherModelRef& InModel);
+
+protected:
+
+	/**
+	 * Refreshes the list of device proxies.
+	 */
+	void RefreshDeviceProxyList();
+
+private:
+
+	// Callback for getting the enabled state of a device proxy list row.
+	bool HandleDeviceListRowIsEnabled(ITargetDeviceProxyPtr DeviceProxy) const;
+
+	// Callback when the user clicks the Device Manager hyperlink.
+	void HandleDeviceManagerHyperlinkNavigate() const;
+
+	// Callback for getting the tool tip text of a device proxy list row.
+	FText HandleDeviceListRowToolTipText(ITargetDeviceProxyPtr DeviceProxy) const;
+
+	// Callback for generating a row in the device list view.
+	TSharedRef<ITableRow> HandleDeviceProxyListViewGenerateRow(ITargetDeviceProxyPtr InItem, const TSharedRef<STableViewBase>& OwnerTable) const;
+
+	// Callback for determining the visibility of the device list view.
+	EVisibility HandleDeviceProxyListViewVisibility() const;
+
+	// Callback for when a device proxy has been added to the device proxy manager.
+	void HandleDeviceProxyManagerProxyAdded(const ITargetDeviceProxyRef& AddedProxy);
+
+	// Callback for when a device proxy has been added to the device proxy manager.
+	void HandleDeviceProxyManagerProxyRemoved(const ITargetDeviceProxyRef& RemovedProxy);
+
+private:
+
+	// Holds the list of available device proxies.
+	TArray<ITargetDeviceProxyPtr> DeviceProxyList;
+
+	// Holds the device proxy list view .
+	TSharedPtr<SListView<ITargetDeviceProxyPtr> > DeviceProxyListView;
+
+	// Holds a pointer to the data model.
+	FProjectLauncherModelPtr Model;
+
+	// Specifies whether advanced options are shown.
+	TAttribute<bool> IsAdvanced;
+
+	// Holds a delegate to be invoked when a profile is run.
+	FOnProfileRun OnProfileRun;
+};
diff --git a/Engine/Source/Developer/ProjectLauncher/Private/Widgets/Launch/SProjectLauncherLaunchCustomRoles.cpp b/Engine/Source/Developer/ProjectLauncher/Private/Widgets/Launch/SProjectLauncherLaunchCustomRoles.cpp
new file mode 100644
--- /dev/null
+++ b/Engine/Source/Developer/ProjectLauncher/Private/Widgets/Launch/SProjectLauncherLaunchCustomRoles.cpp
@@ -0,0 +1,35 @@
+// Copyright 1998-2014 Epic Games, Inc. All Rights Reserved.
+
+#include "ProjectLauncherPrivatePCH.h"
+
+
+#define LOCTEXT_NAMESPACE "SProjectLauncherLaunchCustomRoles"
+
+
+/* SProjectLauncherLaunchCustomRoles structors
+ *****************************************************************************/
+
+SProjectLauncherLaunchCustomRoles::~SProjectLauncherLaunchCustomRoles( )
+{
+	if (Model.IsValid())
+	{
+		Model->OnProfileSelected().RemoveAll(this);
+	}
+}
+
+
+/* SProjectLauncherLaunchCustomRoles interface
+ *****************************************************************************/
+
+void SProjectLauncherLaunchCustomRoles::Construct( const FArguments& InArgs, const FProjectLauncherModelRef& InModel )
+{
+	Model = InModel;
+
+	ChildSlot
+	[
+		SNullWidget::NullWidget
+	];
+}
+
+
+#undef LOCTEXT_NAMESPACE
diff --git a/Engine/Source/Developer/ProjectLauncher/Private/Widgets/Launch/SProjectLauncherLaunchCustomRoles.h b/Engine/Source/Developer/ProjectLauncher/Private/Widgets/Launch/SProjectLauncherLaunchCustomRoles.h
new file mode 100644
--- /dev/null
+++ b/Engine/Source/Developer/ProjectLauncher/Private/Widgets/Launch/SProjectLauncherLaunchCustomRoles.h
@@ -0,0 +1,38 @@
+// Copyright 1998-2014 Epic Games, Inc. All Rights Reserved.
+
+#pragma once
+
+
+/**
+ * Implements the settings panel for launching with custom roles.
+ */
+class SProjectLauncherLaunchCustomRoles
+	: public SCompoundWidget
+{
+public:
+
+	SLATE_BEGIN_ARGS(SProjectLauncherLaunchCustomRoles) { }
+	SLATE_END_ARGS()
+
+public:
+
+	/**
+	 * Destructor.
+	 */
+	~SProjectLauncherLaunchCustomRoles( );
+
+public:
+
+	/**
+	 * Constructs the widget.
+	 *
+	 * @param InArgs - The Slate argument list.
+	 * @param InModel - The data model.
+	 */
+	void Construct( const FArguments& InArgs, const FProjectLauncherModelRef& InModel );
+
+private:
+
+	// Holds a pointer to the data model.
+	FProjectLauncherModelPtr Model;
+};
diff --git a/Engine/Source/Developer/ProjectLauncher/Private/Widgets/Launch/SProjectLauncherLaunchRoleEditor.h b/Engine/Source/Developer/ProjectLauncher/Private/Widgets/Launch/SProjectLauncherLaunchRoleEditor.h
new file mode 100644
--- /dev/null
+++ b/Engine/Source/Developer/ProjectLauncher/Private/Widgets/Launch/SProjectLauncherLaunchRoleEditor.h
@@ -0,0 +1,463 @@
+// Copyright 1998-2014 Epic Games, Inc. All Rights Reserved.
+
+#pragma once
+
+
+#define LOCTEXT_NAMESPACE "SProjectLauncherLaunchRoleEditor"
+
+
+/**
+ * Implements the settings panel for a single launch role.
+ */
+class SProjectLauncherLaunchRoleEditor
+	: public SCompoundWidget
+{
+public:
+
+	SLATE_BEGIN_ARGS(SProjectLauncherLaunchRoleEditor)
+		: _AvailableCultures()
+		, _AvailableMaps()
+	{ }
+		
+		/** The role to be edited initially. */
+		SLATE_ARGUMENT(ILauncherProfileLaunchRolePtr, InitialRole)
+
+		/** The list of available cultures. */
+		SLATE_ARGUMENT(const TArray<FString>*, AvailableCultures)
+
+		/** The list of available maps. */
+		SLATE_ARGUMENT(const TArray<FString>*, AvailableMaps)
+
+	SLATE_END_ARGS()
+
+public:
+
+	/**
+	 * Constructs the widget.
+	 *
+	 * @param InArgs The Slate argument list.
+	 * @param InRole The launch role to edit.
+	 */
+	BEGIN_SLATE_FUNCTION_BUILD_OPTIMIZATION
+	void Construct(	const FArguments& InArgs )
+	{
+		AvailableCultures = InArgs._AvailableCultures;
+		AvailableMaps = InArgs._AvailableMaps;
+		Role = InArgs._InitialRole;
+
+		// create instance types menu
+		FMenuBuilder InstanceTypeMenuBuilder(true, NULL);
+		{
+			FUIAction StandaloneClientAction(FExecuteAction::CreateSP(this, &SProjectLauncherLaunchRoleEditor::HandleInstanceTypeMenuEntryClicked, ELauncherProfileRoleInstanceTypes::StandaloneClient));
+			InstanceTypeMenuBuilder.AddMenuEntry(LOCTEXT("StandaloneClient", "Standalone Client"), LOCTEXT("StandaloneClientActionHint", "Launch this instance as a standalone game client."), FSlateIcon(), StandaloneClientAction);
+
+			FUIAction ListenServerAction(FExecuteAction::CreateSP(this, &SProjectLauncherLaunchRoleEditor::HandleInstanceTypeMenuEntryClicked, ELauncherProfileRoleInstanceTypes::ListenServer));
+			InstanceTypeMenuBuilder.AddMenuEntry(LOCTEXT("ListenServer", "Listen Server"), LOCTEXT("ListenServerActionHint", "Launch this instance as a game client that can accept connections from other clients."), FSlateIcon(), ListenServerAction);
+
+			FUIAction DedicatedServerAction(FExecuteAction::CreateSP(this, &SProjectLauncherLaunchRoleEditor::HandleInstanceTypeMenuEntryClicked, ELauncherProfileRoleInstanceTypes::DedicatedServer));
+			InstanceTypeMenuBuilder.AddMenuEntry(LOCTEXT("DedicatedServer", "Dedicated Server"), LOCTEXT("DedicatedServerActionHint", "Launch this instance as a dedicated game server."), FSlateIcon(), DedicatedServerAction);
+
+			FUIAction UnrealEditorAction(FExecuteAction::CreateSP(this, &SProjectLauncherLaunchRoleEditor::HandleInstanceTypeMenuEntryClicked, ELauncherProfileRoleInstanceTypes::UnrealEditor));
+			InstanceTypeMenuBuilder.AddMenuEntry(LOCTEXT("UnrealEditor", "Unreal Editor"), LOCTEXT("UnrealEditorActionHint", "Launch this instance as an Unreal Editor."), FSlateIcon(), UnrealEditorAction);
+		}
+
+		ChildSlot
+		[
+			SNew(SVerticalBox)
+
+			+ SVerticalBox::Slot()
+				.AutoHeight()
+				[
+					SNew(SProjectLauncherFormLabel)
+						.LabelText(LOCTEXT("InstanceTypeComboBoxLabel", "Launch As:"))
+				]
+
+			+ SVerticalBox::Slot()
+				.AutoHeight()
+				.Padding(0.0f, 2.0f, 0.0f, 0.0f)
+				[
+					// instance type menu
+					SNew(SComboButton)
+						.ButtonContent()
+						[
+							SNew(STextBlock)
+								.Text(this, &SProjectLauncherLaunchRoleEditor::HandleInstanceTypeComboButtonContentText)
+						]
+						.ContentPadding(FMargin(6.0f, 2.0f))
+						.MenuContent()
+						[
+							InstanceTypeMenuBuilder.MakeWidget()
+						]
+				]
+									
+			+ SVerticalBox::Slot()
+				.AutoHeight()
+				.Padding(0.0f, 8.0f, 0.0f, 0.0f)
+				[
+					SNew(SProjectLauncherFormLabel)
+						.ErrorToolTipText(LOCTEXT("CultureNotAvailableError", "The selected culture is not being cooked or is not available."))
+						.ErrorVisibility(this, &SProjectLauncherLaunchRoleEditor::HandleCultureValidationErrorIconVisibility)
+						.LabelText(LOCTEXT("InitialCultureTextBoxLabel", "Initial Culture:"))
+				]
+
+			+ SVerticalBox::Slot()
+				.AutoHeight()
+				.Padding(0.0f, 4.0f, 0.0f, 0.0f)
+				[
+					// initial culture combo box
+					SAssignNew(CultureComboBox, STextComboBox)
+						.ColorAndOpacity(this, &SProjectLauncherLaunchRoleEditor::HandleCultureComboBoxColorAndOpacity)
+						.OptionsSource(&CultureList)
+						.OnSelectionChanged(this, &SProjectLauncherLaunchRoleEditor::HandleCultureComboBoxSelectionChanged)
+				]
+
+			+ SVerticalBox::Slot()
+				.AutoHeight()
+				.Padding(0.0f, 8.0f, 0.0f, 0.0f)
+				[
+					SNew(SProjectLauncherFormLabel)
+						.ErrorToolTipText(LOCTEXT("MapNotAvailableError", "The selected map is not being cooked or is not available."))
+						.ErrorVisibility(this, &SProjectLauncherLaunchRoleEditor::HandleMapValidationErrorIconVisibility)
+						.LabelText(LOCTEXT("InitialMapTextBoxLabel", "Initial Map:"))
+				]
+
+			+ SVerticalBox::Slot()
+				.AutoHeight()
+				.Padding(0.0f, 4.0f, 0.0f, 0.0f)
+				[
+					// initial map combo box
+					SAssignNew(MapComboBox, STextComboBox)
+						.ColorAndOpacity(this, &SProjectLauncherLaunchRoleEditor::HandleMapComboBoxColorAndOpacity)
+						.OptionsSource(&MapList)
+						.OnSelectionChanged(this, &SProjectLauncherLaunchRoleEditor::HandleMapComboBoxSelectionChanged)
+				]
+
+			+ SVerticalBox::Slot()
+				.AutoHeight()
+				.Padding(0.0f, 8.0f, 0.0f, 0.0f)
+				[
+					SNew(SProjectLauncherFormLabel)
+						.LabelText(LOCTEXT("CommandLineTextBoxLabel", "Additional Command Line Parameters:"))
+				]
+
+			+ SVerticalBox::Slot()
+				.AutoHeight()
+				.Padding(0.0f, 4.0f, 0.0f, 0.0f)
+				[
+					// command line text box
+					SAssignNew(CommandLineTextBox, SEditableTextBox)
+						.OnTextChanged(this, &SProjectLauncherLaunchRoleEditor::HandleCommandLineTextBoxTextChanged)
+				]
+
+			+ SVerticalBox::Slot()
+				.AutoHeight()
+				.Padding(0.0f, 12.0f, 0.0f, 0.0f)
+				[
+					// v-sync check box
+					SNew(SCheckBox)
+						.IsChecked(this, &SProjectLauncherLaunchRoleEditor::HandleVsyncCheckBoxIsChecked)
+						.OnCheckStateChanged(this, &SProjectLauncherLaunchRoleEditor::HandleVsyncCheckBoxCheckStateChanged)
+						.Padding(FMargin(4.0f, 0.0f))
+						.Content()
+						[
+							SNew(STextBlock)
+								.Text(LOCTEXT("VsyncCheckBoxText", "Synchronize Screen Refresh Rate (VSync)"))
+						]
+				]
+		];
+
+		Refresh(InArgs._InitialRole);
+	}
+	END_SLATE_FUNCTION_BUILD_OPTIMIZATION
+
+	/**
+	 * Refreshes the widget.
+	 *
+	 * @param InRole The role to edit, or NULL if no role is being edited.
+	 */
+	void Refresh( const ILauncherProfileLaunchRolePtr& InRole )
+	{
+		Role = InRole;
+
+		ILauncherProfileLaunchRolePtr RolePtr = Role.Pin();
+
+		CultureList.Reset();
+		CultureList.Add(MakeShareable(new FString(LOCTEXT("DefaultCultureText", "<default>").ToString())));
+
+		MapList.Reset();
+		MapList.Add(MakeShareable(new FString(LOCTEXT("DefaultMapText", "<default>").ToString())));
+
+		if (RolePtr.IsValid())
+		{
+			// refresh widgets
+			CommandLineTextBox->SetText(FText::FromString(*RolePtr->GetCommandLine()));
+
+			// rebuild culture list
+			if (AvailableCultures != NULL)
+			{
+				for (int32 AvailableCultureIndex = 0; AvailableCultureIndex < AvailableCultures->Num(); ++AvailableCultureIndex)
+				{
+					TSharedPtr<FString> Culture = MakeShareable(new FString((*AvailableCultures)[AvailableCultureIndex]));
+
+					CultureList.Add(Culture);
+
+					if (*Culture == RolePtr->GetInitialCulture())
+					{
+						CultureComboBox->SetSelectedItem(Culture);
+					}
+				}
+			}
+
+			if (RolePtr->GetInitialCulture().IsEmpty() || (CultureList.Num() == 0))
+			{
+				CultureComboBox->SetSelectedItem(CultureList[0]);
+			}
+
+			// rebuild map list
+			if (AvailableMaps != NULL)
+			{
+				for (int32 AvailableMapIndex = 0; AvailableMapIndex < AvailableMaps->Num(); ++AvailableMapIndex)
+				{
+					TSharedPtr<FString> Map = MakeShareable(new FString((*AvailableMaps)[AvailableMapIndex]));
+
+					MapList.Add(Map);
+
+					if (*Map == RolePtr->GetInitialMap())
+					{
+						MapComboBox->SetSelectedItem(Map);
+					}
+				}
+			}
+
+			if (RolePtr->GetInitialMap().IsEmpty() || (MapList.Num() == 0))
+			{
+				MapComboBox->SetSelectedItem(MapList[0]);
+			}
+		}
+		else
+		{
+			CommandLineTextBox->SetText(FText::GetEmpty());
+			CultureComboBox->ClearSelection();
+			MapComboBox->ClearSelection();
+		}
+
+		CultureComboBox->RefreshOptions();
+		MapComboBox->RefreshOptions();
+	}
+
+private:
+
+	// Callback for when the text of the command line text box has changed.
+	void HandleCommandLineTextBoxTextChanged( const FText& InText )
+	{
+		ILauncherProfileLaunchRolePtr RolePtr = Role.Pin();
+
+		if (RolePtr.IsValid())
+		{
+			RolePtr->SetCommandLine(InText.ToString());
+		}
+	}
+
+	// Callback for getting the name of the selected instance type.
+	FString HandleInstanceTypeComboButtonContentText( ) const
+	{
+		ILauncherProfileLaunchRolePtr RolePtr = Role.Pin();
+
+		if (RolePtr.IsValid())
+		{
+			return ELauncherProfileRoleInstanceTypes::ToString(RolePtr->GetInstanceType());
+		}
+
+		return FString();
+	}
+
+	// Callback for clicking an item in the 'Instance type' menu.
+	void HandleInstanceTypeMenuEntryClicked( ELauncherProfileRoleInstanceTypes::Type InstanceType )
+	{
+		ILauncherProfileLaunchRolePtr RolePtr = Role.Pin();
+
+		if (RolePtr.IsValid())
+		{
+			RolePtr->SetInstanceType(InstanceType);
+		}
+	}
+
+	// Callback for getting the foreground text color of the culture combo box.
+	FSlateColor HandleCultureComboBoxColorAndOpacity( ) const
+	{
+		ILauncherProfileLaunchRolePtr RolePtr = Role.Pin();
+
+		if (RolePtr.IsValid())
+		{
+			const FString& InitialCulture = RolePtr->GetInitialCulture();
+
+			if (InitialCulture.IsEmpty() || (AvailableCultures->Contains(InitialCulture)))
+			{
+				return FSlateColor::UseForeground();
+			}
+		}
+
+		return FLinearColor::Red;
+	}
+
+	// Callback for changing the initial culture selection.
+	void HandleCultureComboBoxSelectionChanged( TSharedPtr<FString> Selection, ESelectInfo::Type SelectInfo )
+	{
+		ILauncherProfileLaunchRolePtr RolePtr = Role.Pin();
+
+		if (RolePtr.IsValid())
+		{
+			if (Selection.IsValid() && (Selection != CultureList[0]))
+			{
+				RolePtr->SetInitialCulture(*Selection);
+			}
+			else
+			{
+				RolePtr->SetInitialCulture(FString());
+			}
+		}
+	}
+
+	// Callback for getting the visibility of the validation error icon of the 'Initial Culture' combo box.
+	EVisibility HandleCultureValidationErrorIconVisibility( ) const
+	{
+		ILauncherProfileLaunchRolePtr RolePtr = Role.Pin();
+
+		if (RolePtr.IsValid())
+		{
+			const FString& InitialCulture = RolePtr->GetInitialCulture();
+
+			if (InitialCulture.IsEmpty() || (AvailableCultures->Contains(InitialCulture)))
+			{
+				return EVisibility::Hidden;
+			}
+		}
+
+		return EVisibility::Visible;
+	}
+
+	// Callback for getting the visibility of the validation error icon of the 'Launch As' combo box.
+	EVisibility HandleLaunchAsValidationErrorIconVisibility( ) const
+	{
+		return EVisibility::Hidden;
+	}
+
+	// Callback for getting the foreground text color of the map combo box.
+	FSlateColor HandleMapComboBoxColorAndOpacity( ) const
+	{
+		ILauncherProfileLaunchRolePtr RolePtr = Role.Pin();
+
+		if (RolePtr.IsValid())
+		{
+			const FString& InitialMap = RolePtr->GetInitialMap();
+
+			if (InitialMap.IsEmpty() || (AvailableMaps->Contains(InitialMap)))
+			{
+				return FSlateColor::UseForeground();
+			}
+		}
+
+		return FLinearColor::Red;
+	}
+
+	// Callback for changing the initial map selection.
+	void HandleMapComboBoxSelectionChanged( TSharedPtr<FString> Selection, ESelectInfo::Type SelectInfo )
+	{
+		ILauncherProfileLaunchRolePtr RolePtr = Role.Pin();
+
+		if (RolePtr.IsValid())
+		{
+			if (Selection.IsValid() && (Selection != MapList[0]))
+			{
+				RolePtr->SetInitialMap(*Selection);
+			}
+			else
+			{
+				RolePtr->SetInitialMap(FString());
+			}
+		}
+	}
+
+	// Callback for getting the visibility of the validation error icon of the 'Initial Map' combo box.
+	EVisibility HandleMapValidationErrorIconVisibility( ) const
+	{
+		ILauncherProfileLaunchRolePtr RolePtr = Role.Pin();
+
+		if (RolePtr.IsValid())
+		{
+			const FString& InitialMap = RolePtr->GetInitialMap();
+
+			if (InitialMap.IsEmpty() || (AvailableMaps->Contains(InitialMap)))
+			{
+				return EVisibility::Hidden;
+			}
+		}
+
+		return EVisibility::Visible;
+	}
+
+	// Callback for checking the specified 'No VSync' check box.
+	void HandleVsyncCheckBoxCheckStateChanged( ESlateCheckBoxState::Type NewState )
+	{
+		ILauncherProfileLaunchRolePtr RolePtr = Role.Pin();
+
+		if (RolePtr.IsValid())
+		{
+			RolePtr->SetVsyncEnabled(NewState == ESlateCheckBoxState::Checked);
+		}
+	}
+
+	// Callback for determining the checked state of the specified 'No VSync' check box.
+	ESlateCheckBoxState::Type HandleVsyncCheckBoxIsChecked( ) const
+	{
+		ILauncherProfileLaunchRolePtr RolePtr = Role.Pin();
+
+		if (RolePtr.IsValid())
+		{
+			if (RolePtr->IsVsyncEnabled())
+			{
+				return ESlateCheckBoxState::Checked;
+			}
+		}
+
+		return ESlateCheckBoxState::Unchecked;
+	}
+
+private:
+
+	// Holds a pointer to the list of available cultures.
+	const TArray<FString>* AvailableCultures;
+
+	// Holds a pointer to the list of available maps.
+	const TArray<FString>* AvailableMaps;
+
+	// Holds the command line text box.
+	TSharedPtr<SEditableTextBox> CommandLineTextBox;
+
+	// Holds the instance type combo box.
+	TSharedPtr<SComboBox<TSharedPtr<ELauncherProfileRoleInstanceTypes::Type> > > InstanceTypeComboBox;
+
+	// Holds the list of instance types.
+	TArray<TSharedPtr<ELauncherProfileRoleInstanceTypes::Type> > InstanceTypeList;
+
+	// Holds the culture text box.
+	TSharedPtr<STextComboBox> CultureComboBox;
+
+	// Holds the list of cultures that are available for the selected game.
+	TArray<TSharedPtr<FString> > CultureList;
+
+	// Holds the map text box.
+	TSharedPtr<STextComboBox> MapComboBox;
+
+	// Holds the list of maps that are available for the selected game.
+	TArray<TSharedPtr<FString> > MapList;
+
+	// Holds the profile that owns the role being edited.
+	TWeakPtr<ILauncherProfile> Profile;
+
+	// Holds a pointer to the role that is being edited in this widget.
+	TWeakPtr<ILauncherProfileLaunchRole> Role;
+};
+
+
+#undef LOCTEXT_NAMESPACE
diff --git a/Engine/Source/Developer/ProjectLauncher/Private/Widgets/Package/SProjectLauncherPackagingSettings.cpp b/Engine/Source/Developer/ProjectLauncher/Private/Widgets/Package/SProjectLauncherPackagingSettings.cpp
new file mode 100644
--- /dev/null
+++ b/Engine/Source/Developer/ProjectLauncher/Private/Widgets/Package/SProjectLauncherPackagingSettings.cpp
@@ -0,0 +1,145 @@
+// Copyright 1998-2014 Epic Games, Inc. All Rights Reserved.
+
+#include "ProjectLauncherPrivatePCH.h"
+
+
+#define LOCTEXT_NAMESPACE "SProjectLauncherPackagingSettings"
+
+
+/* SProjectLauncherPackagingSettings structors
+ *****************************************************************************/
+
+SProjectLauncherPackagingSettings::~SProjectLauncherPackagingSettings( )
+{
+	if (Model.IsValid())
+	{
+		Model->OnProfileSelected().RemoveAll(this);
+	}
+}
+
+
+/* SProjectLauncherPackagingSettings interface
+ *****************************************************************************/
+
+BEGIN_SLATE_FUNCTION_BUILD_OPTIMIZATION
+void SProjectLauncherPackagingSettings::Construct( const FArguments& InArgs, const FProjectLauncherModelRef& InModel )
+{
+	Model = InModel;
+
+	ChildSlot
+	[
+		SNew(SVerticalBox)
+
+		+ SVerticalBox::Slot()
+			.FillHeight(1.0)
+			[
+				SNew(SBorder)
+					.Padding(8.0)
+					.BorderImage(FEditorStyle::GetBrush("ToolPanel.GroupBorder"))
+					[
+						SNew(SVerticalBox)
+
+						+ SVerticalBox::Slot()
+							.AutoHeight()
+							[
+								SNew(STextBlock)
+									.Text(LOCTEXT("RepositoryPathLabel", "Repository Path:"))
+							]
+
+						+ SVerticalBox::Slot()
+							.AutoHeight()
+							.Padding(0.0, 4.0, 0.0, 0.0)
+							[
+								SNew(SHorizontalBox)
+
+								+ SHorizontalBox::Slot()
+									.FillWidth(1.0)
+									.Padding(0.0, 0.0, 0.0, 3.0)
+									[
+										// repository path text box
+										SAssignNew(RepositoryPathTextBox, SEditableTextBox)
+									]
+
+								+ SHorizontalBox::Slot()
+									.AutoWidth()
+									.HAlign(HAlign_Right)
+									.Padding(4.0, 0.0, 0.0, 0.0)
+									[
+										// browse button
+										SNew(SButton)
+											.ContentPadding(FMargin(6.0, 2.0))
+											.IsEnabled(false)
+											.Text(LOCTEXT("BrowseButtonText", "Browse..."))
+											.ToolTipText(LOCTEXT("BrowseButtonToolTip", "Browse for the repository"))
+											//.OnClicked(this, &SProjectLauncherPackagePage::HandleBrowseButtonClicked)
+									]
+							]
+
+						+ SVerticalBox::Slot()
+							.AutoHeight()
+							.Padding(0.0, 8.0, 0.0, 0.0)
+							[
+								SNew(STextBlock)
+									.Text(LOCTEXT("DescriptionTextBoxLabel", "Description:"))
+							]
+
+						+ SVerticalBox::Slot()
+							.FillHeight(1.0)
+							.Padding(0.0, 4.0, 0.0, 0.0)
+							[
+								SNew(SEditableTextBox)
+							]
+					]
+			]
+
+		+ SVerticalBox::Slot()
+			.AutoHeight()
+			.Padding(0.0, 8.0, 0.0, 0.0)
+			[
+				SNew(SExpandableArea)
+					.AreaTitle(LOCTEXT("IncrementalPackagingAreaTitle", "Incremental Packaging"))
+					.InitiallyCollapsed(true)
+					.Padding(8.0)
+					.BodyContent()
+					[
+						SNew(SVerticalBox)
+
+						+ SVerticalBox::Slot()
+							.AutoHeight()
+							[
+								// incremental packaging check box
+								SNew(SCheckBox)
+									//.IsChecked(this, &SProjectLauncherLaunchPage::HandleIncrementalCheckBoxIsChecked)
+									//.OnCheckStateChanged(this, &SProjectLauncherLaunchPage::HandleIncrementalCheckBoxCheckStateChanged)
+									.Content()
+									[
+										SNew(STextBlock)
+											.Text(LOCTEXT("IncrementalCheckBoxText", "This is an incremental build"))
+									]
+							]
+
+						+ SVerticalBox::Slot()
+							.FillHeight(1.0)
+							.Padding(18.0, 12.0, 0.0, 0.0)
+							[
+								SNew(SVerticalBox)
+							]
+					]
+			]
+	];
+
+	Model->OnProfileSelected().AddSP(this, &SProjectLauncherPackagingSettings::HandleProfileManagerProfileSelected);
+}
+END_SLATE_FUNCTION_BUILD_OPTIMIZATION
+
+
+/* SProjectLauncherPackagingSettings callbacks
+ *****************************************************************************/
+
+void SProjectLauncherPackagingSettings::HandleProfileManagerProfileSelected( const ILauncherProfilePtr& SelectedProfile, const ILauncherProfilePtr& PreviousProfile )
+{
+
+}
+
+
+#undef LOCTEXT_NAMESPACE
diff --git a/Engine/Source/Developer/ProjectLauncher/Private/Widgets/Profile/SProjectLauncherProfileListView.h b/Engine/Source/Developer/ProjectLauncher/Private/Widgets/Profile/SProjectLauncherProfileListView.h
new file mode 100644
--- /dev/null
+++ b/Engine/Source/Developer/ProjectLauncher/Private/Widgets/Profile/SProjectLauncherProfileListView.h
@@ -0,0 +1,119 @@
+// Copyright 1998-2014 Epic Games, Inc. All Rights Reserved.
+
+#pragma once
+
+/**
+ * Implements the deployment targets panel.
+ */
+class SProjectLauncherProfileListView
+	: public SCompoundWidget
+{
+public:
+
+	SLATE_BEGIN_ARGS(SProjectLauncherProfileListView) { }
+		/**
+		 * The Callback for when a profile is to be edited.
+		 */
+		SLATE_EVENT(FOnProfileRun, OnProfileEdit)
+
+		/**
+		 * The Callback for when a profile is to be run.
+		 */
+		SLATE_EVENT(FOnProfileRun, OnProfileRun)
+
+		/**
+		* The Callback for when a profile is to be deleted.
+		*/
+		SLATE_EVENT(FOnProfileRun, OnProfileDelete)
+	SLATE_END_ARGS()
+
+public:
+	/**
+	 * Constructor.
+	 */
+	SProjectLauncherProfileListView();
+
+	/**
+	 * Destructor.
+	 */
+	~SProjectLauncherProfileListView( );
+
+public:
+
+	/**
+	 * Constructs the widget.
+	 *
+	 * @param InArgs The Slate argument list.
+	 * @param InModel The data model.
+	 */
+	void Construct(const FArguments& InArgs, const FProjectLauncherModelRef& InModel);
+
+protected:
+
+	/**
+	 * Refreshes the list of launch profiles.
+	 */
+	void RefreshLaunchProfileList();
+
+private:
+
+	// Builds the command list for the context menu on list items.
+	void CreateCommands();
+
+	// Callback for getting the enabled state of a profile list row.
+	bool HandleProfileRowIsEnabled(ILauncherProfilePtr LaunchProfile) const;
+
+	// Callback for getting the tool tip text of a profile list row.
+	FText HandleDeviceListRowToolTipText(ILauncherProfilePtr LaunchProfile) const;
+
+	// Callback for generating a row in the profile list view.
+	TSharedRef<ITableRow> HandleProfileListViewGenerateRow(ILauncherProfilePtr InItem, const TSharedRef<STableViewBase>& OwnerTable) const;
+
+	// Callback for when a launch profile has been added to the profile manager.
+	void HandleProfileManagerProfileAdded(const ILauncherProfileRef& AddedProfile);
+
+	// Callback for when a launch profile has been removed to the profile manager.
+	void HandleProfileManagerProfileRemoved(const ILauncherProfileRef& RemovedProfile);
+
+	// Callback for creating a context menu for the list.
+	TSharedPtr<SWidget> MakeProfileContextMenu();
+
+	// Callback returns true if the rename command can be executed.
+	bool HandleRenameProfileCommandCanExecute() const;
+
+	// Callback to execute the rename command from the context menu.
+	void HandleRenameProfileCommandExecute();
+
+	// Callback returns true if the duplicate command can be executed.
+	bool HandleDuplicateProfileCommandCanExecute() const;
+
+	// Callback to execute the duplicate command from the context menu.
+	void HandleDuplicateProfileCommandExecute();
+
+	// Callback returns true if the delete command can be executed.
+	bool HandleDeleteProfileCommandCanExecute() const;
+
+	// Callback to execute the delete command from the context menu.
+	void HandleDeleteProfileCommandExecute();
+
+private:
+
+	// Holds the launch profile list view .
+	TSharedPtr<SListView<ILauncherProfilePtr> > LaunchProfileListView;
+
+	// Holds a pointer to the data model.
+	FProjectLauncherModelPtr Model;
+
+	// Holds a delegate to be invoked when a profile is to be edited.
+	FOnProfileRun OnProfileEdit;
+
+	// Holds a delegate to be invoked when a profile is run.
+	FOnProfileRun OnProfileRun;
+
+	// Holds a delegate to be invoked when a profile is deleted.
+	FOnProfileRun OnProfileDelete;
+
+	/** Commands handled by this widget */
+	TSharedRef<FUICommandList> CommandList;
+
+};
diff --git a/Engine/Source/Developer/ProjectLauncher/Private/Widgets/Project/SProjectLauncherProjectPage.h b/Engine/Source/Developer/ProjectLauncher/Private/Widgets/Project/SProjectLauncherProjectPage.h
new file mode 100644
--- /dev/null
+++ b/Engine/Source/Developer/ProjectLauncher/Private/Widgets/Project/SProjectLauncherProjectPage.h
@@ -0,0 +1,37 @@
+// Copyright 1998-2014 Epic Games, Inc. All Rights Reserved.
+
+#pragma once
+
+
+/**
+ * Implements the profile page for the session launcher wizard.
+ */
+class SProjectLauncherProjectPage
+	: public SCompoundWidget
+{
+public:
+
+	SLATE_BEGIN_ARGS(SProjectLauncherProjectPage) { }
+		SLATE_ATTRIBUTE(ILauncherProfilePtr, LaunchProfile)
+	SLATE_END_ARGS()
+
+public:
+
+	/**
+	 * Constructs the widget.
+	 *
+	 * @param InArgs The Slate argument list.
+	 * @param InModel The data model.
+	 */
+	void Construct(	const FArguments& InArgs, const FProjectLauncherModelRef& InModel, bool InShowConfig = true );
+
+private:
+
+	// Callback for getting the enabled state of the project settings area.
+	bool HandleProjectSettingsAreaIsEnabled( ) const;
+
+private:
+
+	// Holds a pointer to the data model.
+	FProjectLauncherModelPtr Model;
+};
diff --git a/Engine/Source/Developer/ProjectLauncher/Private/Widgets/Project/SProjectLauncherProjectPicker.h b/Engine/Source/Developer/ProjectLauncher/Private/Widgets/Project/SProjectLauncherProjectPicker.h
new file mode 100644
--- /dev/null
+++ b/Engine/Source/Developer/ProjectLauncher/Private/Widgets/Project/SProjectLauncherProjectPicker.h
@@ -0,0 +1,99 @@
+// Copyright 1998-2014 Epic Games, Inc. All Rights Reserved.
+
+#pragma once
+
+/**
+ * Implements the project loading area for the session launcher wizard.
+ */
+class SProjectLauncherProjectPicker
+	: public SCompoundWidget
+{
+public:
+
+	SLATE_BEGIN_ARGS(SProjectLauncherProjectPicker) { }
+		SLATE_ATTRIBUTE(ILauncherProfilePtr, LaunchProfile)
+	SLATE_END_ARGS()
+
+public:
+
+	/**
+	 * Destructor.
+	 */
+	~SProjectLauncherProjectPicker( );
+
+public:
+
+	/**
+	 * Constructs the widget.
+	 *
+	 * @param InArgs The Slate argument list.
+	 * @param InModel The data model.
+	 */
+	void Construct(const FArguments& InArgs, const FProjectLauncherModelRef& InModel);
+
+protected:
+
+	/**
+	 * Creates the widget for the project menu.
+	 *
+	 * @return The widget.
+	 */
+	TSharedRef<SWidget> MakeProjectMenuWidget( );
+
+	/**
+	 * Creates the widget for the project selection.
+	 *
+	 * @return The widget.
+	 */
+	TSharedRef<SWidget> MakeProjectWidget( );
+
+	/**
+	 * Creates the widget for the project and configuration selection.
+	 *
+	 * @return The widget.
+	 */
+	TSharedRef<SWidget> MakeProjectAndConfigWidget( );
+
+private:
+
+	// Handles selecting a build configuration.
+	void HandleBuildConfigurationSelectorConfigurationSelected( EBuildConfigurations::Type );
+
+	// Handles getting the content text of the build configuration selector.
+	FString HandleBuildConfigurationSelectorText( ) const;
+
+	// Handles changing the selected profile in the profile manager.
+	void HandleProfileManagerProfileSelected( const ILauncherProfilePtr& SelectedProfile, const ILauncherProfilePtr& PreviousProfile );
+
+	// Handles getting the text for the project combo button.
+	FString HandleProjectComboButtonText( ) const;
+
+	// Handles getting the tooltip for the project combo button.
+	FString HandleProjectComboButtonToolTip( ) const;
+
+	// Handles clicking the "any project" option.
+	void HandleAnyProjectClicked(FString ProjectPath);
+
+	// Handles clicking a project menu entry.
+	void HandleProjectMenuEntryClicked( FString ProjectPath );
+
+	// Handles determining the visibility of a validation error icon.
+	EVisibility HandleValidationErrorIconVisibility( ELauncherProfileValidationErrors::Type Error ) const;
+
+	// Sets the project in the appropriate place (profile if provided otherwise on the model)
+	void SetProjectPath(FString ProjectPath);
+
+private:
+
+	// Attribute for the launch profile this widget edits, if null it edits the project in the launcher model
+	TAttribute<ILauncherProfilePtr> LaunchProfileAttr;
+
+	// Holds the list of available projects.
+	TArray<TSharedPtr<FString>> ProjectList;
+
+	// Holds a pointer to the data model.
+	FProjectLauncherModelPtr Model;
+
+	// Holds the repository path text box.
+	TSharedPtr<SEditableTextBox> RepositoryPathTextBox;
+};
diff --git a/Engine/Source/Developer/ProjectLauncher/Private/Widgets/SProjectLauncher.cpp b/Engine/Source/Developer/ProjectLauncher/Private/Widgets/SProjectLauncher.cpp
new file mode 100644
--- /dev/null
+++ b/Engine/Source/Developer/ProjectLauncher/Private/Widgets/SProjectLauncher.cpp
@@ -0,0 +1,319 @@
+// Copyright 1998-2014 Epic Games, Inc. All Rights Reserved.
+
+#include "ProjectLauncherPrivatePCH.h"
+
+
+#define LOCTEXT_NAMESPACE "SProjectLauncher"
+
+
+/* SProjectLauncher structors
+ *****************************************************************************/
+
+SProjectLauncher::SProjectLauncher()
+	:bAdvanced(false)
+{
+	if (GConfig != NULL)
+	{
+		GConfig->GetBool(TEXT("FProjectLauncher"), TEXT("AdvancedMode"), bAdvanced, GEngineIni);
+	}
+}
+
+SProjectLauncher::~SProjectLauncher( )
+{
+	if (GConfig != NULL)
+	{
+		GConfig->SetBool(TEXT("FProjectLauncher"), TEXT("AdvancedMode"), bAdvanced, GEngineIni);
+	}
+
+	if (LauncherWorker.IsValid())
+	{
+		LauncherWorker->Cancel();
+		FPlatformProcess::Sleep(0.5f);
+	}
+}
+
+
+/* SProjectLauncher interface
+ *****************************************************************************/
+
+BEGIN_SLATE_FUNCTION_BUILD_OPTIMIZATION
+void SProjectLauncher::Construct( const FArguments& InArgs, const TSharedRef<SDockTab>& ConstructUnderMajorTab, const TSharedPtr<SWindow>& ConstructUnderWindow, const FProjectLauncherModelRef& InModel )
+{
+	Model = InModel;
+
+	// create & initialize main menu bar
+	TSharedRef<FWorkspaceItem> RootMenuGroup = FWorkspaceItem::NewGroup(LOCTEXT("RootMenuGroup", "Root"));
+
+	FMenuBarBuilder MenuBarBuilder = FMenuBarBuilder(TSharedPtr<FUICommandList>());
+	MenuBarBuilder.AddPullDownMenu(
+		LOCTEXT("WindowMenuLabel", "Window"),
+		FText::GetEmpty(),
+		FNewMenuDelegate::CreateStatic(&SProjectLauncher::FillWindowMenu, RootMenuGroup),
+		"Window"
+	);
+
+	ChildSlot
+	[
+		SAssignNew(WidgetSwitcher, SWidgetSwitcher)
+		.WidgetIndex((int32)ELauncherPanels::Launch)
+		
+		// Empty Panel
+		+ SWidgetSwitcher::Slot()
+		[
+			SNew(SBorder)
+		]
+
+		// SProjectLauncher Panel
+		+ SWidgetSwitcher::Slot()
+		[
+			SNew(SSplitter)
+			.Style(FEditorStyle::Get(), "ContentBrowser.Splitter")
+			.Orientation(Orient_Vertical)
+
+			// Simple SProjectLauncher
+			+ SSplitter::Slot()
+			.Value(0.7f)
+			[
+				SNew(SVerticalBox)
+
+				+ SVerticalBox::Slot()
+				.AutoHeight()
+				.Padding(2)
+				[
+					SNew(SHorizontalBox)
+					
+					// Project Bar
+					+ SHorizontalBox::Slot()
+					.AutoWidth()
+					[
+						SNew(SProjectLauncherProjectPicker, InModel)
+					]
+
+					// Advanced Button
+					+ SHorizontalBox::Slot()
+					.HAlign(HAlign_Right)
+					[
+						SNew(SBorder)
+						.Padding(2)
+						.BorderImage(FEditorStyle::GetBrush("ToolPanel.GroupBorder"))
+						[
+							SNew(SCheckBox)
+							.Style(FEditorStyle::Get(), "ToggleButtonCheckbox")
+							.IsFocusable(true)
+							.ToolTipText(LOCTEXT("ToggleAdvancedOptionsToolTipText", "Toggles Advanced Options"))
+							.OnCheckStateChanged(this, &SProjectLauncher::OnAdvancedChanged)
+							.IsChecked(this, &SProjectLauncher::OnIsAdvanced)
+							[
+								SNew(SHorizontalBox)
+
+								// Icon
+								+ SHorizontalBox::Slot()
+								.VAlign(VAlign_Center)
+								[
+									SNew(SImage)
+									.Image(this, &SProjectLauncher::GetAdvancedToggleBrush)
+								]
+
+								// Text
+								+ SHorizontalBox::Slot()
+								.AutoWidth()
+								.VAlign(VAlign_Center)
+								.Padding(4,0,4,0)
+								[
+									SNew(STextBlock)
+									.TextStyle(FCoreStyle::Get(), "Toolbar.Label")
+									.ShadowOffset(FVector2D::UnitVector)
+									.Text(LOCTEXT("AdvancedButton", "Advanced"))
+								]
+							]
+						]
+					]
+				]
+
+				+ SVerticalBox::Slot()
+				.FillHeight(1)
+				.Padding(2)
+				[
+					SNew(SBorder)
+					[
+						SAssignNew(LaunchList, SProjectLauncherSimpleDeviceListView, InModel)
+						.OnProfileRun(this, &SProjectLauncher::OnProfileRun)
+						.IsAdvanced(this, &SProjectLauncher::GetIsAdvanced)
+					]
+				]
+			]
+
+			+ SSplitter::Slot()
+			.Value(0.3f)
+			[
+				SNew(SBorder)
+				[
+					SNew(SVerticalBox)
+
+					+ SVerticalBox::Slot()
+					.AutoHeight()
+					.Padding(2)
+					[
+						SNew(SBorder)
+						.BorderImage(FEditorStyle::GetBrush("ToolPanel.GroupBorder"))
+						.Padding(4.0)
+						[
+							SNew(SHorizontalBox)
+
+							+ SHorizontalBox::Slot()
+							.AutoWidth()
+							[
+								SNew(STextBlock)
+								.TextStyle(FCoreStyle::Get(), "Toolbar.Label")
+								.ShadowOffset(FVector2D::UnitVector)
+								.Text(LOCTEXT("ProjectLauncherCustomProfilesTitle", "Custom Launch Profiles"))
+							]
+
+							+ SHorizontalBox::Slot()
+							.HAlign(HAlign_Right)
+							[
+								SNew(SButton)
+								.ButtonStyle(FEditorStyle::Get(), "ToggleButton")
+								.ToolTipText(LOCTEXT("ProjectLauncherCustomProfileAdd", "Add a new custom launch profile."))
+								.ContentPadding(0)
+								.OnClicked(this, &SProjectLauncher::OnAddCustomLaunchProfileClicked)
+								[
+									SNew(SImage)
+									.Image(FCoreStyle::Get().GetBrush("EditableComboBox.Add"))
+									.ColorAndOpacity(FSlateColor::UseForeground())
+								]
+							]
+						]
+					]
+					+ SVerticalBox::Slot()
+					.FillHeight(1)
+					.Padding(2)
+					[
+						SAssignNew(ProfileList, SBorder)
+						.BorderImage(FEditorStyle::GetBrush("NoBorder"))
+						.Padding(0)
+						[
+							// Simple Launch List
+							SNew(SProjectLauncherProfileListView, InModel)
+							.OnProfileEdit(this, &SProjectLauncher::OnProfileEdit)
+							.OnProfileRun(this, &SProjectLauncher::OnProfileRun)
+							.OnProfileDelete(this, &SProjectLauncher::OnProfileDelete)
+						]
+					]
+				]
+			]
+		]
+		
+		// Launch Settings
+		+ SWidgetSwitcher::Slot()
+		[
+			SAssignNew(ProfileSettingsPanel, SProjectLauncherSettings, InModel)
+			.OnCloseClicked(this, &SProjectLauncher::OnProfileSettingsClose)
+			.OnDeleteClicked(this, &SProjectLauncher::OnProfileDelete)
+		]
+
+		// Progress Panel
+		+ SWidgetSwitcher::Slot()
+		[
+			SAssignNew(ProgressPanel, SProjectLauncherProgress)
+			.OnCloseClicked(this, &SProjectLauncher::OnProgressClose)
+		]
+	];
+}
+END_SLATE_FUNCTION_BUILD_OPTIMIZATION
+
+/* SProjectLauncher implementation
+*****************************************************************************/
+
+void SProjectLauncher::FillWindowMenu( FMenuBuilder& MenuBuilder, TSharedRef<FWorkspaceItem> RootMenuGroup )
+{
+#if !WITH_EDITOR
+	MenuBuilder.BeginSection("WindowGlobalTabSpawners", LOCTEXT("UfeMenuGroup", "Unreal Frontend"));
+	{
+		FGlobalTabmanager::Get()->PopulateTabSpawnerMenu(MenuBuilder, RootMenuGroup);
+	}
+	MenuBuilder.EndSection();
+#endif //!WITH_EDITOR
+}
+
+/* SProjectLauncher callbacks
+*****************************************************************************/
+
+void SProjectLauncher::OnAdvancedChanged(const ESlateCheckBoxState::Type NewCheckedState)
+{
+	bAdvanced = (NewCheckedState == ESlateCheckBoxState::Checked);
+}
+
+ESlateCheckBoxState::Type SProjectLauncher::OnIsAdvanced() const
+{
+	return (bAdvanced) ? ESlateCheckBoxState::Checked : ESlateCheckBoxState::Unchecked;
+}
+
+const FSlateBrush* SProjectLauncher::GetAdvancedToggleBrush() const
+{
+	return FEditorStyle::GetBrush("LauncherCommand.AdvancedBuild.Medium");
+}
+
+bool SProjectLauncher::GetIsAdvanced() const
+{
+	return bAdvanced;
+}
+
+void SProjectLauncher::OnProfileEdit(const ILauncherProfileRef& Profile)
+{
+	Model->SelectProfile(Profile);
+	WidgetSwitcher->SetActiveWidgetIndex(ELauncherPanels::ProfileEditor);
+}
+
+void SProjectLauncher::OnProfileRun(const ILauncherProfileRef& Profile)
+{
+	LauncherWorker = Model->GetSProjectLauncher()->Launch(Model->GetDeviceProxyManager(), Profile);
+	
+	if (LauncherWorker.IsValid())
+	{
+		ProgressPanel->SetLauncherWorker(LauncherWorker.ToSharedRef());
+		WidgetSwitcher->SetActiveWidgetIndex(ELauncherPanels::Progress);
+	}
+}
+
+void SProjectLauncher::OnProfileDelete(const ILauncherProfileRef& Profile)
+{
+	Model->GetProfileManager()->RemoveProfile(Profile);
+}
+
+FReply SProjectLauncher::OnAddCustomLaunchProfileClicked()
+{
+	ILauncherProfileRef Profile = Model->GetProfileManager()->AddNewProfile();
+	
+	OnProfileEdit(Profile);
+
+	ProfileSettingsPanel->EnterEditMode();
+
+	return FReply::Handled();
+}
+
+FReply SProjectLauncher::OnProfileSettingsClose()
+{
+	if (LauncherWorker.IsValid())
+	{
+		LauncherWorker->Cancel();
+	}
+
+	WidgetSwitcher->SetActiveWidgetIndex(ELauncherPanels::Launch);
+
+	return FReply::Handled();
+}
+
+FReply SProjectLauncher::OnProgressClose()
+{
+	if (LauncherWorker.IsValid())
+	{
+		LauncherWorker->Cancel();
+	}
+
+	WidgetSwitcher->SetActiveWidgetIndex(ELauncherPanels::Launch);
+
+	return FReply::Handled();
+}
+
+#undef LOCTEXT_NAMESPACE
diff --git a/Engine/Source/Developer/ProjectLauncher/Private/Widgets/SProjectLauncher.h b/Engine/Source/Developer/ProjectLauncher/Private/Widgets/SProjectLauncher.h
new file mode 100644
--- /dev/null
+++ b/Engine/Source/Developer/ProjectLauncher/Private/Widgets/SProjectLauncher.h
@@ -0,0 +1,121 @@
+// Copyright 1998-2014 Epic Games, Inc. All Rights Reserved.
+
+#pragma once
+
+
+/**
+ * Implements a Slate widget for the launcher user interface.
+ */
+class SProjectLauncher
+	: public SCompoundWidget
+{
+public:
+
+	SLATE_BEGIN_ARGS(SProjectLauncher){ }
+
+		/** Exposes a delegate to be invoked when the launcher has closed. */
+		SLATE_EVENT(FSimpleDelegate, OnClosed)
+
+	SLATE_END_ARGS()
+
+public:
+
+	/**
+	* Constructor.
+	*/
+	SProjectLauncher();
+
+	/**
+	 * Destructor.
+	 */
+	~SProjectLauncher( );
+
+public:
+
+	/**
+	 * Constructs the widget.
+	 *
+	 * @param InArgs The Slate argument list.
+	 * @param ConstructUnderMajorTab The major tab which will contain the session front-end.
+	 * @param ConstructUnderWindow The window in which this widget is being constructed.
+	 * @param InModel The view model to use.
+	 */
+	void Construct( const FArguments& InArgs, const TSharedRef<SDockTab>& ConstructUnderMajorTab, const TSharedPtr<SWindow>& ConstructUnderWindow, const FProjectLauncherModelRef& InModel );
+
+	void ShowWidget(int32 Index);
+
+protected:
+
+	/**
+	 * Fills the Window menu with menu items.
+	 *
+	 * @param MenuBuilder The multi-box builder that should be filled with content for this pull-down menu.
+	 * @param RootMenuGroup The root menu group.
+	 */
+	static void FillWindowMenu( FMenuBuilder& MenuBuilder, TSharedRef<FWorkspaceItem> RootMenuGroup );
+
+	/**
+	 * Launches the selected profile.
+	 */
+	void LaunchSelectedProfile( );
+
+private:
+
+	// Callback for toggling the advanced mode.
+	void OnAdvancedChanged(const ESlateCheckBoxState::Type NewCheckedState);
+
+	// Callback to determine if we are in advanced mode.
+	ESlateCheckBoxState::Type OnIsAdvanced() const;
+
+	// Get advanced toggle button brush.
+	const FSlateBrush* GetAdvancedToggleBrush() const;
+
+	// Callback for whether advanced options should be shown.
+	bool GetIsAdvanced() const;
+
+	// Callback for editing a profile
+	void OnProfileEdit(const ILauncherProfileRef& Profile);
+
+	// Callback for running a profile
+	void OnProfileRun(const ILauncherProfileRef& Profile);
+
+	// Callback for deleting a profile
+	void OnProfileDelete(const ILauncherProfileRef& Profile);
+
+	// Callback for clicking the Add New Launch Profile
+	FReply OnAddCustomLaunchProfileClicked();
+
+	// Callback for when the settings panel is closed
+	FReply OnProfileSettingsClose();
+
+	// Callback for when the progress panel is closed
+	FReply OnProgressClose();
+
+private:
+
+	// Holds the current launcher worker, if any.
+	ILauncherWorkerPtr LauncherWorker;
+
+	// Holds a pointer to the view model.
+	FProjectLauncherModelPtr Model;
+
+	// Holds the profile settings panel.
+	TSharedPtr<SProjectLauncherSettings> ProfileSettingsPanel;
+
+	// Holds the progress panel.
+	TSharedPtr<SProjectLauncherProgress> ProgressPanel;
+
+	// Holds the widget switcher.
+	TSharedPtr<SWidgetSwitcher> WidgetSwitcher;
+
+	// Contains the launch list widgets.
+	//TSharedPtr<SBorder> LaunchList;
+
+	TSharedPtr<SProjectLauncherSimpleDeviceListView> LaunchList;
+
+	// Contains the profile list widgets.
+	TSharedPtr<SBorder> ProfileList;
+
+	// Whether we are showing advanced options.
+	bool bAdvanced;
+};
diff --git a/Engine/Source/Developer/ProjectLauncher/Private/Widgets/SProjectLauncherBuildTaskSettings.h b/Engine/Source/Developer/ProjectLauncher/Private/Widgets/SProjectLauncherBuildTaskSettings.h
new file mode 100644
--- /dev/null
+++ b/Engine/Source/Developer/ProjectLauncher/Private/Widgets/SProjectLauncherBuildTaskSettings.h
@@ -0,0 +1,38 @@
+// Copyright 1998-2014 Epic Games, Inc. All Rights Reserved.
+
+#pragma once
+
+
+/**
+ * Implements the launcher settings widget.
+ */
+class SProjectLauncherBuildTaskSettings
+	: public SCompoundWidget
+{
+public:
+
+	SLATE_BEGIN_ARGS(SProjectLauncherBuildTaskSettings) { }
+	SLATE_END_ARGS()
+
+public:
+
+	/**
+	 * Destructor.
+	 */
+	~SProjectLauncherBuildTaskSettings( );
+
+public:
+
+	/**
+	 * Constructs the widget.
+	 *
+	 * @param InArgs The Slate argument list.
+	 * @param InModel The data model.
+	 */
+	void Construct( const FArguments& InArgs, const FProjectLauncherModelRef& InModel );
+
+private:
+
+	// Holds a pointer to the data model.
+	FProjectLauncherModelPtr Model;
+};
diff --git a/Engine/Source/Developer/ProjectLauncher/Private/Widgets/SProjectLauncherDelegates.h b/Engine/Source/Developer/ProjectLauncher/Private/Widgets/SProjectLauncherDelegates.h
new file mode 100644
--- /dev/null
+++ b/Engine/Source/Developer/ProjectLauncher/Private/Widgets/SProjectLauncherDelegates.h
@@ -0,0 +1,10 @@
+// Copyright 1998-2014 Epic Games, Inc. All Rights Reserved.
+
+#pragma once
+
+/**
+ * Delegate type for SProjectLauncher Profile running.
+ *
+ * The first parameter is the SProjectLauncher Profile that you want to run.
+ */
+DECLARE_DELEGATE_OneParam(FOnProfileRun, const ILauncherProfileRef&)
\ No newline at end of file
diff --git a/Engine/Source/Developer/ProjectLauncher/Private/Widgets/SProjectLauncherDeployTaskSettings.h b/Engine/Source/Developer/ProjectLauncher/Private/Widgets/SProjectLauncherDeployTaskSettings.h
new file mode 100644
--- /dev/null
+++ b/Engine/Source/Developer/ProjectLauncher/Private/Widgets/SProjectLauncherDeployTaskSettings.h
@@ -0,0 +1,38 @@
+// Copyright 1998-2014 Epic Games, Inc. All Rights Reserved.
+
+#pragma once
+
+
+/**
+ * Implements the launcher settings widget.
+ */
+class SProjectLauncherDeployTaskSettings
+	: public SCompoundWidget
+{
+public:
+
+	SLATE_BEGIN_ARGS(SProjectLauncherDeployTaskSettings) { }
+	SLATE_END_ARGS()
+
+public:
+
+	/**
+	 * Destructor.
+	 */
+	~SProjectLauncherDeployTaskSettings( );
+
+public:
+
+	/**
+	 * Constructs the widget.
+	 *
+	 * @param InArgs The Slate argument list.
+	 * @param InModel The data model.
+	 */
+	void Construct( const FArguments& InArgs, const FProjectLauncherModelRef& InModel );
+
+private:
+
+	// Holds a pointer to the data model.
+	FProjectLauncherModelPtr Model;
+};
