Added missing build configuration UI in the advanced settings.
Fixed possible crash when removing a device.
Some cleanup of unused API.

diff --git a/Engine/Source/Developer/ProjectLauncher/Private/Widgets/Build/SProjectLauncherBuildPage.cpp b/Engine/Source/Developer/ProjectLauncher/Private/Widgets/Build/SProjectLauncherBuildPage.cpp
--- a/Engine/Source/Developer/ProjectLauncher/Private/Widgets/Build/SProjectLauncherBuildPage.cpp
+++ b/Engine/Source/Developer/ProjectLauncher/Private/Widgets/Build/SProjectLauncherBuildPage.cpp
@@ -57,6 +57,37 @@ void SProjectLauncherBuildPage::Construct( const FArguments& InArgs, const FProj
 					]
 			]
 
+		+ SVerticalBox::Slot()
+			.AutoHeight()
+			.Padding(0, 3, 0, 3)
+			[
+				SNew(SBorder)
+				.BorderImage(FEditorStyle::GetBrush("ToolPanel.GroupBorder"))
+				.Visibility(this, &SProjectLauncherBuildPage::ShowBuildConfiguration)
+				[
+					SNew(SHorizontalBox)
+
+					+ SHorizontalBox::Slot()
+					.AutoWidth()
+					[
+						SNew(SProjectLauncherFormLabel)
+						.ErrorToolTipText(NSLOCTEXT("SProjectLauncherBuildValidation", "NoBuildConfigurationSelectedError", "A Build Configuration must be selected."))
+						.ErrorVisibility(this, &SProjectLauncherBuildPage::HandleValidationErrorIconVisibility, ELauncherProfileValidationErrors::NoBuildConfigurationSelected)
+						.LabelText(LOCTEXT("ConfigurationComboBoxLabel", "Build Configuration:"))
+					]
+
+					+ SHorizontalBox::Slot()
+					.AutoWidth()
+					[
+						// build configuration selector
+						SNew(SProjectLauncherBuildConfigurationSelector)
+						.Font(FCoreStyle::Get().GetFontStyle(TEXT("NormalFont")))
+						.OnConfigurationSelected(this, &SProjectLauncherBuildPage::HandleBuildConfigurationSelectorConfigurationSelected)
+						.Text(this, &SProjectLauncherBuildPage::HandleBuildConfigurationSelectorText)
+					]
+				]
+			]
+
         + SVerticalBox::Slot()
             .AutoHeight()
             .Padding(0.0f, 8.0f, 0.0f, 0.0f)
@@ -205,5 +236,55 @@ bool SProjectLauncherBuildPage::HandleGenDSYMButtonEnabled() const
     return false;
 }
 
+EVisibility SProjectLauncherBuildPage::ShowBuildConfiguration() const
+{
+	ILauncherProfilePtr SelectedProfile = Model->GetSelectedProfile();
+
+	if (SelectedProfile.IsValid() && SelectedProfile->IsBuilding())
+	{
+		return EVisibility::Visible;
+	}
+	else
+	{
+		return EVisibility::Collapsed;
+	}
+}
+
+void SProjectLauncherBuildPage::HandleBuildConfigurationSelectorConfigurationSelected(EBuildConfigurations::Type Configuration)
+{
+	ILauncherProfilePtr SelectedProfile = Model->GetSelectedProfile();
+
+	if (SelectedProfile.IsValid())
+	{
+		SelectedProfile->SetBuildConfiguration(Configuration);
+	}
+}
+
+FString SProjectLauncherBuildPage::HandleBuildConfigurationSelectorText() const
+{
+	ILauncherProfilePtr SelectedProfile = Model->GetSelectedProfile();
+
+	if (SelectedProfile.IsValid())
+	{
+		return EBuildConfigurations::ToString(SelectedProfile->GetBuildConfiguration());
+	}
+
+	return FString();
+}
+
+EVisibility SProjectLauncherBuildPage::HandleValidationErrorIconVisibility(ELauncherProfileValidationErrors::Type Error) const
+{
+	ILauncherProfilePtr SelectedProfile = Model->GetSelectedProfile();
+
+	if (SelectedProfile.IsValid())
+	{
+		if (SelectedProfile->HasValidationError(Error))
+		{
+			return EVisibility::Visible;
+		}
+	}
+
+	return EVisibility::Hidden;
+}
 
 #undef LOCTEXT_NAMESPACE
diff --git a/Engine/Source/Developer/ProjectLauncher/Private/Widgets/Build/SProjectLauncherBuildPage.h b/Engine/Source/Developer/ProjectLauncher/Private/Widgets/Build/SProjectLauncherBuildPage.h
--- a/Engine/Source/Developer/ProjectLauncher/Private/Widgets/Build/SProjectLauncherBuildPage.h
+++ b/Engine/Source/Developer/ProjectLauncher/Private/Widgets/Build/SProjectLauncherBuildPage.h
@@ -62,6 +62,18 @@ private:
 	// Callback for getting the enabled state of the Generate DSYM button.
     bool HandleGenDSYMButtonEnabled() const;
 
+	// Callback for determining if the build configuration should be shown
+	EVisibility ShowBuildConfiguration() const;
+
+	// Callback for selecting a build configuration.
+	void HandleBuildConfigurationSelectorConfigurationSelected(EBuildConfigurations::Type Configuration);
+
+	// Callback for getting the content text of the build configuration selector.
+	FString HandleBuildConfigurationSelectorText() const;
+
+	// Callback for determining the visibility of a validation error icon.
+	EVisibility HandleValidationErrorIconVisibility(ELauncherProfileValidationErrors::Type Error) const;
+
 private:
 
 	// Holds a pointer to the data model.
diff --git a/Engine/Source/Developer/ProjectLauncher/Private/Widgets/Project/SProjectLauncherProjectPage.h b/Engine/Source/Developer/ProjectLauncher/Private/Widgets/Project/SProjectLauncherProjectPage.h
--- a/Engine/Source/Developer/ProjectLauncher/Private/Widgets/Project/SProjectLauncherProjectPage.h
+++ b/Engine/Source/Developer/ProjectLauncher/Private/Widgets/Project/SProjectLauncherProjectPage.h
@@ -25,11 +25,6 @@ public:
 	 */
 	void Construct(	const FArguments& InArgs, const FProjectLauncherModelRef& InModel, bool InShowConfig = true );
 
-private:
-
-	// Callback for getting the enabled state of the project settings area.
-	bool HandleProjectSettingsAreaIsEnabled( ) const;
-
 private:
 
 	// Holds a pointer to the data model.
diff --git a/Engine/Source/Developer/ProjectLauncher/Private/Widgets/Project/SProjectLauncherProjectPicker.h b/Engine/Source/Developer/ProjectLauncher/Private/Widgets/Project/SProjectLauncherProjectPicker.h
--- a/Engine/Source/Developer/ProjectLauncher/Private/Widgets/Project/SProjectLauncherProjectPicker.h
+++ b/Engine/Source/Developer/ProjectLauncher/Private/Widgets/Project/SProjectLauncherProjectPicker.h
@@ -47,24 +47,8 @@ protected:
 	 */
 	TSharedRef<SWidget> MakeProjectWidget( );
 
-	/**
-	 * Creates the widget for the project and configuration selection.
-	 *
-	 * @return The widget.
-	 */
-	TSharedRef<SWidget> MakeProjectAndConfigWidget( );
-
 private:
 
-	// Handles selecting a build configuration.
-	void HandleBuildConfigurationSelectorConfigurationSelected( EBuildConfigurations::Type );
-
-	// Handles getting the content text of the build configuration selector.
-	FString HandleBuildConfigurationSelectorText( ) const;
-
-	// Handles changing the selected profile in the profile manager.
-	void HandleProfileManagerProfileSelected( const ILauncherProfilePtr& SelectedProfile, const ILauncherProfilePtr& PreviousProfile );
-
 	// Handles getting the text for the project combo button.
 	FString HandleProjectComboButtonText( ) const;
 
@@ -94,6 +78,4 @@ private:
 	// Holds a pointer to the data model.
 	FProjectLauncherModelPtr Model;
 
-	// Holds the repository path text box.
-	TSharedPtr<SEditableTextBox> RepositoryPathTextBox;
 };
diff --git a/Engine/Source/Developer/ProjectLauncher/Private/Widgets/Shared/SProjectLauncherBuildConfigurationSelector.h b/Engine/Source/Developer/ProjectLauncher/Private/Widgets/Shared/SProjectLauncherBuildConfigurationSelector.h
--- a/Engine/Source/Developer/ProjectLauncher/Private/Widgets/Shared/SProjectLauncherBuildConfigurationSelector.h
+++ b/Engine/Source/Developer/ProjectLauncher/Private/Widgets/Shared/SProjectLauncherBuildConfigurationSelector.h
@@ -25,6 +25,7 @@ public:
 	SLATE_BEGIN_ARGS(SProjectLauncherBuildConfigurationSelector) { }
 		SLATE_EVENT(FOnSessionSProjectLauncherBuildConfigurationSelected, OnConfigurationSelected)
 		SLATE_ATTRIBUTE(FString, Text)
+		SLATE_ATTRIBUTE(FSlateFontInfo, Font)
 	SLATE_END_ARGS()
 
 public:
@@ -35,7 +36,7 @@ public:
 	 * @param InArgs The Slate argument list.
 	 * @param InModel The data model.
 	 */
-	void Construct(	const FArguments& InArgs )
+	void Construct(const FArguments& InArgs)
 	{
 		OnConfigurationSelected = InArgs._OnConfigurationSelected;
 
@@ -58,6 +59,8 @@ public:
 			MenuBuilder.AddMenuEntry(LOCTEXT("TestAction", "Test"), LOCTEXT("TestActionHint", "Test configuration."), FSlateIcon(), TestAction);
 		}
 
+		FSlateFontInfo TextFont = InArgs._Font.IsSet() ? InArgs._Font.Get() : FCoreStyle::Get().GetFontStyle(TEXT("SmallFont"));
+		
 		ChildSlot
 		[
 			// build configuration menu
@@ -66,7 +69,7 @@ public:
 			.ButtonContent()
 			[
 				SNew(STextBlock)
-				.Font(FCoreStyle::Get().GetFontStyle(TEXT("SmallFont")))
+				.Font(TextFont)
 				.Text(InArgs._Text)
 			]
 			.ContentPadding(FMargin(4.0f, 2.0f))
diff --git a/Engine/Source/Developer/TargetDeviceServices/Private/TargetDeviceService.cpp b/Engine/Source/Developer/TargetDeviceServices/Private/TargetDeviceService.cpp
--- a/Engine/Source/Developer/TargetDeviceServices/Private/TargetDeviceService.cpp
+++ b/Engine/Source/Developer/TargetDeviceServices/Private/TargetDeviceService.cpp
@@ -6,6 +6,11 @@
 
 void FTargetDeviceService::AddTargetDevice(ITargetDevicePtr InDevice)
 {
+	if (!InDevice.IsValid())
+	{
+		return;
+	}
+
 	FName Variant = FName(InDevice->GetTargetPlatform().PlatformName().GetCharArray().GetData());
 
 	if (DevicePlatformName == NAME_None)
@@ -35,18 +40,37 @@ void FTargetDeviceService::AddTargetDevice(ITargetDevicePtr InDevice)
 	// Sort and choose cache the default
 	TargetDevicePtrs.ValueSort(FVariantSortCallback());
 	auto DeviceIterator = TargetDevicePtrs.CreateIterator();
-	DefaultDevicePtr = (*DeviceIterator).Value;
+	if (DeviceIterator)
+	{
+		DefaultDevicePtr = (*DeviceIterator).Value;
+	}
+	else
+	{
+		DefaultDevicePtr = nullptr;
+	}
 }
 
 void FTargetDeviceService::RemoveTargetDevice(ITargetDevicePtr InDevice)
 {
+	if (!InDevice.IsValid())
+	{
+		return;
+	}
+
 	FName Variant = FName(InDevice->GetTargetPlatform().PlatformName().GetCharArray().GetData());
 
 	TargetDevicePtrs.Remove(Variant);
 
 	// Cache the default
 	auto DeviceIterator = TargetDevicePtrs.CreateIterator();
-	DefaultDevicePtr = (*DeviceIterator).Value;
+	if (DeviceIterator)
+	{
+		DefaultDevicePtr = (*DeviceIterator).Value;
+	}
+	else
+	{
+		DefaultDevicePtr = nullptr;
+	}
 }
 
 void FTargetDeviceService::HandlePingMessage(const FTargetDeviceServicePing& InMessage, const IMessageContextRef& Context)
