Fix for project launcher on mac.

diff --git a/Engine/Source/Developer/ProjectLauncher/Private/Models/ProjectLauncherCommands.h b/Engine/Source/Developer/ProjectLauncher/Private/Models/ProjectLauncherCommands.h
--- a/Engine/Source/Developer/ProjectLauncher/Private/Models/ProjectLauncherCommands.h
+++ b/Engine/Source/Developer/ProjectLauncher/Private/Models/ProjectLauncherCommands.h
@@ -15,8 +15,8 @@ public:
 	 */
 	FProjectLauncherCommands( )
 		: TCommands<FProjectLauncherCommands>(
-			"ProjectLauncherCommand",
-			NSLOCTEXT("Contexts", "ProjectLauncherCommand", "Project Launcher Command"),
+			"LauncherCommand",
+			NSLOCTEXT("Contexts", "LauncherCommand", "Project Launcher Command"),
 			NAME_None,
 			FEditorStyle::GetStyleSetName()
 	)
diff --git a/Engine/Source/Developer/ProjectLauncher/Private/Widgets/SProjectLauncher.cpp b/Engine/Source/Developer/ProjectLauncher/Private/Widgets/SProjectLauncher.cpp
--- a/Engine/Source/Developer/ProjectLauncher/Private/Widgets/SProjectLauncher.cpp
+++ b/Engine/Source/Developer/ProjectLauncher/Private/Widgets/SProjectLauncher.cpp
@@ -39,6 +39,8 @@ SProjectLauncher::~SProjectLauncher( )
 BEGIN_SLATE_FUNCTION_BUILD_OPTIMIZATION
 void SProjectLauncher::Construct( const FArguments& InArgs, const TSharedRef<SDockTab>& ConstructUnderMajorTab, const TSharedPtr<SWindow>& ConstructUnderWindow, const FProjectLauncherModelRef& InModel )
 {
+	FProjectLauncherCommands::Register();
+
 	Model = InModel;
 
 	// create & initialize main menu bar
diff --git a/Engine/Source/Developer/ProjectLauncher/Private/Widgets/SProjectLauncherSettings.cpp b/Engine/Source/Developer/ProjectLauncher/Private/Widgets/SProjectLauncherSettings.cpp
--- a/Engine/Source/Developer/ProjectLauncher/Private/Widgets/SProjectLauncherSettings.cpp
+++ b/Engine/Source/Developer/ProjectLauncher/Private/Widgets/SProjectLauncherSettings.cpp
@@ -20,8 +20,6 @@ SProjectLauncherSettings::~SProjectLauncherSettings( )
 BEGIN_SLATE_FUNCTION_BUILD_OPTIMIZATION
 void SProjectLauncherSettings::Construct( const FArguments& InArgs, const FProjectLauncherModelRef& InModel )
 {
-	FProjectLauncherCommands::Register();
-
 	Model = InModel;
 
 	OnCloseClicked = InArgs._OnCloseClicked;
