Removed Whitespace from usernames.
Added new buttons on the lauch profile results screen ( rerun, and cancel without closing).
Fixed possible crash when trying to populate a non existant profiler window.

diff --git a/Engine/Source/Developer/AutomationWindow/Private/SAutomationWindow.cpp b/Engine/Source/Developer/AutomationWindow/Private/SAutomationWindow.cpp
--- a/Engine/Source/Developer/AutomationWindow/Private/SAutomationWindow.cpp
+++ b/Engine/Source/Developer/AutomationWindow/Private/SAutomationWindow.cpp
@@ -1484,7 +1484,7 @@ void SAutomationWindow::FindWorkers()
 {
 	ActiveSession = SessionManager->GetSelectedSession();
 
-	bool SessionIsValid = ActiveSession.IsValid() && (ActiveSession->GetSessionOwner() == FPlatformProcess::UserName(false));
+	bool SessionIsValid = ActiveSession.IsValid() && (ActiveSession->GetSessionOwner() == FPlatformProcess::UserName(true));
 
 	if (SessionIsValid)
 	{
diff --git a/Engine/Source/Developer/DeviceManager/Private/Widgets/Browser/SDeviceBrowserDeviceListRow.h b/Engine/Source/Developer/DeviceManager/Private/Widgets/Browser/SDeviceBrowserDeviceListRow.h
--- a/Engine/Source/Developer/DeviceManager/Private/Widgets/Browser/SDeviceBrowserDeviceListRow.h
+++ b/Engine/Source/Developer/DeviceManager/Private/Widgets/Browser/SDeviceBrowserDeviceListRow.h
@@ -131,7 +131,7 @@ private:
 	{
 		FString ClaimUser = DeviceService->GetClaimUser();
 
-		if ((ClaimUser == FPlatformProcess::UserName(false)) && !DeviceService->IsRunning())
+		if ((ClaimUser == FPlatformProcess::UserName(true)) && !DeviceService->IsRunning())
 		{
 			ClaimUser += LOCTEXT("RemotelyHint", " (remotely)").ToString();
 		}
diff --git a/Engine/Source/Developer/LauncherServices/Private/Launcher/LauncherDeployCommands.h b/Engine/Source/Developer/LauncherServices/Private/Launcher/LauncherDeployCommands.h
--- a/Engine/Source/Developer/LauncherServices/Private/Launcher/LauncherDeployCommands.h
+++ b/Engine/Source/Developer/LauncherServices/Private/Launcher/LauncherDeployCommands.h
@@ -74,7 +74,7 @@ public:
 			*InitialMap,
 			*InstanceId.ToString(),
 			*ChainState.SessionId.ToString(),
-			FPlatformProcess::UserName(false),
+			FPlatformProcess::UserName(true),
 			*ChainState.Profile->GetName(),
 			CookCommand.IsValid() ? *(TEXT(" ") + CookCommand->GetAdditionalArguments(ChainState)) : TEXT(""),
 			((TargetPlatform.PlatformName() == TEXT("PS4") || ChainState.Profile->IsPackingWithUnrealPak())
@@ -193,7 +193,7 @@ public:
 			*InitialMap,
 			*InstanceId.ToString(),
 			*ChainState.SessionId.ToString(),
-			FPlatformProcess::UserName(false),
+			FPlatformProcess::UserName(true),
 			*ChainState.Profile->GetName(),
 			CookCommand.IsValid() ? *(TEXT(" ") + CookCommand->GetAdditionalArguments(ChainState)) : TEXT(""),
 			((TargetPlatform.PlatformName() == TEXT("PS4") || ChainState.Profile->IsPackingWithUnrealPak())
@@ -276,7 +276,7 @@ public:
 			*InitialMap,
 			*InstanceId.ToString(),
 			*ChainState.SessionId.ToString(),
-			FPlatformProcess::UserName(false),
+			FPlatformProcess::UserName(true),
 			*ChainState.Profile->GetName(),
 			CookCommand.IsValid() ? *(TEXT(" ") + CookCommand->GetAdditionalArguments(ChainState)) : TEXT(""),
 			((TargetPlatform.PlatformName() == TEXT("PS4") || ChainState.Profile->IsPackingWithUnrealPak())
@@ -391,7 +391,7 @@ public:
 			*InitialMap,
 			*InstanceId.ToString(),
 			*ChainState.SessionId.ToString(),
-			FPlatformProcess::UserName(false),
+			FPlatformProcess::UserName(true),
 			*ChainState.Profile->GetName(),
 			CookCommand.IsValid() ? *(TEXT(" ") + CookCommand->GetAdditionalArguments(ChainState)) : TEXT(""),
 			((TargetPlatform.PlatformName() == TEXT("PS4") || ChainState.Profile->IsPackingWithUnrealPak())
diff --git a/Engine/Source/Developer/LauncherServices/Private/Launcher/LauncherLaunchCommands.h b/Engine/Source/Developer/LauncherServices/Private/Launcher/LauncherLaunchCommands.h
--- a/Engine/Source/Developer/LauncherServices/Private/Launcher/LauncherLaunchCommands.h
+++ b/Engine/Source/Developer/LauncherServices/Private/Launcher/LauncherLaunchCommands.h
@@ -66,7 +66,7 @@ public:
 		CommandLine += FString::Printf(TEXT(" -addcmdline=\"-InstanceId=%s -SessionId=%s -SessionOwner=%s -SessionName='%s'%s%s%s\""),
 			*InstanceId.ToString(),
 			*ChainState.SessionId.ToString(),
-			FPlatformProcess::UserName(false),
+			FPlatformProcess::UserName(true),
 			*ChainState.Profile->GetName(),
 			CookCommand.IsValid() ? *(TEXT(" ") + CookCommand->GetAdditionalArguments(ChainState)) : TEXT(""),
 			Role->IsVsyncEnabled() ? TEXT(" -vsync") : TEXT(""),
@@ -170,7 +170,7 @@ public:
 		CommandLine += FString::Printf(TEXT(" -addcmdline=\"-InstanceId=%s -SessionId=%s -SessionOwner=%s -SessionName='%s'%s%s\""),
 			*InstanceId.ToString(),
 			*ChainState.SessionId.ToString(),
-			FPlatformProcess::UserName(false),
+			FPlatformProcess::UserName(true),
 			*ChainState.Profile->GetName(),
 			CookCommand.IsValid() ? *(TEXT(" ") + CookCommand->GetAdditionalArguments(ChainState)) : TEXT(""),
 			*Role->GetCommandLine()
diff --git a/Engine/Source/Developer/LauncherServices/Private/Launcher/LauncherPackageCommands.h b/Engine/Source/Developer/LauncherServices/Private/Launcher/LauncherPackageCommands.h
--- a/Engine/Source/Developer/LauncherServices/Private/Launcher/LauncherPackageCommands.h
+++ b/Engine/Source/Developer/LauncherServices/Private/Launcher/LauncherPackageCommands.h
@@ -57,7 +57,7 @@ public:
 			*InitialMap,
 			*InstanceId.ToString(),
 			*ChainState.SessionId.ToString(),
-			FPlatformProcess::UserName(false),
+			FPlatformProcess::UserName(true),
 			*ChainState.Profile->GetName(),
 			CookCommand.IsValid() ? *CookCommand->GetAdditionalArguments(ChainState) : TEXT(""),
 			CookCommand.IsValid() ? *(TEXT(" ") + CookCommand->GetAdditionalArguments(ChainState)) : TEXT(""),
@@ -148,7 +148,7 @@ public:
 			*InitialMap,
 			*InstanceId.ToString(),
 			*ChainState.SessionId.ToString(),
-			FPlatformProcess::UserName(false),
+			FPlatformProcess::UserName(true),
 			*ChainState.Profile->GetName(),
 			CookCommand.IsValid() ? *(TEXT(" ") + CookCommand->GetAdditionalArguments(ChainState)) : TEXT(""),
 			((TargetPlatform.PlatformName() == TEXT("PS4") || ChainState.Profile->IsPackingWithUnrealPak())
@@ -229,7 +229,7 @@ public:
 			*InitialMap,
 			*InstanceId.ToString(),
 			*ChainState.SessionId.ToString(),
-			FPlatformProcess::UserName(false),
+			FPlatformProcess::UserName(true),
 			*ChainState.Profile->GetName(),
 			CookCommand.IsValid() ? *(TEXT(" ") + CookCommand->GetAdditionalArguments(ChainState)) : TEXT(""),
 			((TargetPlatform.PlatformName() == TEXT("PS4") || ChainState.Profile->IsPackingWithUnrealPak())
@@ -334,7 +334,7 @@ public:
 			*InitialMap,
 			*InstanceId.ToString(),
 			*ChainState.SessionId.ToString(),
-			FPlatformProcess::UserName(false),
+			FPlatformProcess::UserName(true),
 			*ChainState.Profile->GetName(),
 			CookCommand.IsValid() ? *(TEXT(" ") + CookCommand->GetAdditionalArguments(ChainState)) : TEXT(""),
 			((TargetPlatform.PlatformName() == TEXT("PS4") || ChainState.Profile->IsPackingWithUnrealPak())
diff --git a/Engine/Source/Developer/Profiler/Private/ProfilerManager.cpp b/Engine/Source/Developer/Profiler/Private/ProfilerManager.cpp
--- a/Engine/Source/Developer/Profiler/Private/ProfilerManager.cpp
+++ b/Engine/Source/Developer/Profiler/Private/ProfilerManager.cpp
@@ -398,7 +398,7 @@ void FProfilerManager::SessionManager_OnSelectedSessionChanged( const ISessionIn
 void FProfilerManager::SessionManager_OnInstanceSelectionChanged()
 {
 	const ISessionInfoPtr& SelectedSession = SessionManager->GetSelectedSession();
-	const bool SessionIsValid = SelectedSession.IsValid() && (SelectedSession->GetSessionOwner() == FPlatformProcess::UserName(false));
+	const bool SessionIsValid = SelectedSession.IsValid() && (SelectedSession->GetSessionOwner() == FPlatformProcess::UserName(true));
 
 	if( ActiveSession != SelectedSession || FProfilerManager::GetSettings().bSingleInstanceMode )
 	{
@@ -438,7 +438,11 @@ void FProfilerManager::SessionManager_OnInstanceSelectionChanged()
 
 				ProfilerSessionInstances.Add( ProfilerSession->GetInstanceID(), ProfilerSession );
 				ProfilerClient->Track( ProfilerInstanceID );
-				GetProfilerWindow()->ManageEventGraphTab( ProfilerInstanceID, true, ProfilerSession->GetName() );
+				TSharedPtr<SProfilerWindow> ProfilerWindow = GetProfilerWindow();
+				if (ProfilerWindow.IsValid())
+				{
+					ProfilerWindow->ManageEventGraphTab(ProfilerInstanceID, true, ProfilerSession->GetName());
+				}
 			}
 		}
 	}
diff --git a/Engine/Source/Developer/ProjectLauncher/Private/Widgets/SProjectLauncher.cpp b/Engine/Source/Developer/ProjectLauncher/Private/Widgets/SProjectLauncher.cpp
--- a/Engine/Source/Developer/ProjectLauncher/Private/Widgets/SProjectLauncher.cpp
+++ b/Engine/Source/Developer/ProjectLauncher/Private/Widgets/SProjectLauncher.cpp
@@ -219,6 +219,7 @@ void SProjectLauncher::Construct( const FArguments& InArgs, const TSharedRef<SDo
 		[
 			SAssignNew(ProgressPanel, SProjectLauncherProgress)
 			.OnCloseClicked(this, &SProjectLauncher::OnProgressClose)
+			.OnRerunClicked(this, &SProjectLauncher::OnRerunClicked)
 		]
 	];
 }
@@ -269,6 +270,7 @@ void SProjectLauncher::OnProfileEdit(const ILauncherProfileRef& Profile)
 
 void SProjectLauncher::OnProfileRun(const ILauncherProfileRef& Profile)
 {
+	LauncherProfile = Profile;
 	LauncherWorker = Model->GetSProjectLauncher()->Launch(Model->GetDeviceProxyManager(), Profile);
 	
 	if (LauncherWorker.IsValid())
@@ -300,6 +302,7 @@ FReply SProjectLauncher::OnProfileSettingsClose()
 	{
 		LauncherWorker->Cancel();
 	}
+	LauncherProfile.Reset();
 
 	WidgetSwitcher->SetActiveWidgetIndex(ELauncherPanels::Launch);
 
@@ -312,10 +315,27 @@ FReply SProjectLauncher::OnProgressClose()
 	{
 		LauncherWorker->Cancel();
 	}
+	LauncherProfile.Reset();
 
 	WidgetSwitcher->SetActiveWidgetIndex(ELauncherPanels::Launch);
 
 	return FReply::Handled();
 }
 
+FReply SProjectLauncher::OnRerunClicked()
+{
+	if (LauncherWorker.IsValid())
+	{
+		LauncherWorker->Cancel();
+	}
+	LauncherWorker = Model->GetSProjectLauncher()->Launch(Model->GetDeviceProxyManager(), LauncherProfile.ToSharedRef());
+
+	if (LauncherWorker.IsValid())
+	{
+		ProgressPanel->SetLauncherWorker(LauncherWorker.ToSharedRef());
+	}
+
+	return FReply::Handled();
+}
+
 #undef LOCTEXT_NAMESPACE
diff --git a/Engine/Source/Developer/ProjectLauncher/Private/Widgets/SProjectLauncher.h b/Engine/Source/Developer/ProjectLauncher/Private/Widgets/SProjectLauncher.h
--- a/Engine/Source/Developer/ProjectLauncher/Private/Widgets/SProjectLauncher.h
+++ b/Engine/Source/Developer/ProjectLauncher/Private/Widgets/SProjectLauncher.h
@@ -91,11 +91,17 @@ private:
 	// Callback for when the progress panel is closed
 	FReply OnProgressClose();
 
+	// Callback for when the progress panel requests the profile to be re run
+	FReply OnRerunClicked();
+
 private:
 
 	// Holds the current launcher worker, if any.
 	ILauncherWorkerPtr LauncherWorker;
 
+	// Holds the Launcher profile the launcherWorker is running.
+	ILauncherProfilePtr LauncherProfile;
+
 	// Holds a pointer to the view model.
 	FProjectLauncherModelPtr Model;
 
diff --git a/Engine/Source/Developer/ProjectLauncher/Private/Widgets/SProjectLauncherProgress.h b/Engine/Source/Developer/ProjectLauncher/Private/Widgets/SProjectLauncherProgress.h
--- a/Engine/Source/Developer/ProjectLauncher/Private/Widgets/SProjectLauncherProgress.h
+++ b/Engine/Source/Developer/ProjectLauncher/Private/Widgets/SProjectLauncherProgress.h
@@ -16,6 +16,7 @@ public:
 
 	SLATE_BEGIN_ARGS(SProjectLauncherProgress) { }
 	SLATE_EVENT(FOnClicked, OnCloseClicked)
+	SLATE_EVENT(FOnClicked, OnRerunClicked)
 	SLATE_END_ARGS()
 
 public:
@@ -38,6 +39,7 @@ public:
 	void Construct( const FArguments& InArgs )
 	{
 		OnCloseClicked = InArgs._OnCloseClicked;
+		OnRerunClicked = InArgs._OnRerunClicked;
 
 		ChildSlot
 		[
@@ -174,14 +176,31 @@ public:
 				.AutoWidth()
 				.Padding(4.0f, 0.0f, 0.0f, 0.0f)
 				[
-					// done button
+					// Re-Run button
 					SNew(SButton)
 					.ContentPadding(FMargin(6.0f, 2.0f))
+					.IsEnabled(this, &SProjectLauncherProgress::IsRerunButtonEnabled)
+					.OnClicked(this, &SProjectLauncherProgress::HandleRerunButtonClicked)
+					.ToolTipText(this, &SProjectLauncherProgress::GetRerunButtonToolTip)
+					[
+						SNew(STextBlock)
+						.Text(this, &SProjectLauncherProgress::GetRerunButtonText)
+					]
+				]
+
+				+ SHorizontalBox::Slot()
+				.AutoWidth()
+				.Padding(4.0f, 0.0f, 0.0f, 0.0f)
+				[
+					// Cancel / Done button
+					SNew(SButton)
+					.ContentPadding(FMargin(6.0f, 2.0f))
+					.IsEnabled(this, &SProjectLauncherProgress::IsDoneButtonEnabled)
 					.OnClicked(this, &SProjectLauncherProgress::HandleDoneButtonClicked)
-					.ToolTipText(LOCTEXT("DoneButtonTooltip", "Close this page."))
+					.ToolTipText(this, &SProjectLauncherProgress::GetDoneButtonToolTip)
 					[
 						SNew(STextBlock)
-						.Text(LOCTEXT("DoneButtonLabel", "Done"))
+						.Text(this, &SProjectLauncherProgress::GetDoneButtonText)
 					]
 				]
 			]
@@ -352,9 +371,64 @@ private:
 		return FReply::Handled();
 	}
 
+	bool IsRerunButtonEnabled() const
+	{
+		ILauncherWorkerPtr LauncherWorkerPtr = LauncherWorker.Pin();
+
+		if (LauncherWorkerPtr.IsValid() && 
+			(LauncherWorkerPtr->GetStatus() == ELauncherWorkerStatus::Canceled || LauncherWorkerPtr->GetStatus() == ELauncherWorkerStatus::Completed))
+		{
+			return true;
+		}
+
+		return false;
+	}
+
+	FReply HandleRerunButtonClicked()
+	{
+		if (OnRerunClicked.IsBound())
+		{
+			return OnRerunClicked.Execute();
+		}
+
+		return FReply::Handled();
+	}
+
+	FText GetRerunButtonToolTip() const
+	{
+		return LOCTEXT("DoneButtonCloseTooltip", "Run this launch profile.");
+	}
+
+	FText GetRerunButtonText() const
+	{
+		return LOCTEXT("DoneButtonDoneLabel", "Run");
+	}
+
+	bool IsDoneButtonEnabled() const
+	{
+		ILauncherWorkerPtr LauncherWorkerPtr = LauncherWorker.Pin();
+
+		if (LauncherWorkerPtr.IsValid() && LauncherWorkerPtr->GetStatus() == ELauncherWorkerStatus::Canceling)
+		{
+			return false;
+		}
+		
+		return true;
+	}
+
 	FReply HandleDoneButtonClicked()
 	{
-		if (OnCloseClicked.IsBound())
+		ILauncherWorkerPtr LauncherWorkerPtr = LauncherWorker.Pin();
+
+		if (LauncherWorkerPtr.IsValid() && LauncherWorkerPtr->GetStatus() == ELauncherWorkerStatus::Busy)
+		{
+			LauncherWorkerPtr->Cancel();
+		}
+		else if (LauncherWorkerPtr.IsValid() && LauncherWorkerPtr->GetStatus() == ELauncherWorkerStatus::Canceling)
+		{
+			// do nothing
+		}
+		else if (OnCloseClicked.IsBound())
 		{
 			return OnCloseClicked.Execute();
 		}
@@ -362,6 +436,42 @@ private:
 		return FReply::Handled();
 	}
 
+	FText GetDoneButtonToolTip() const
+	{
+		ILauncherWorkerPtr LauncherWorkerPtr = LauncherWorker.Pin();
+
+		if (LauncherWorkerPtr.IsValid())
+		{
+			if (LauncherWorkerPtr->GetStatus() == ELauncherWorkerStatus::Busy)
+			{
+				return LOCTEXT("DoneButtonCancelTooltip", "Cancel the run of this profile.");
+			}
+			else if (LauncherWorkerPtr->GetStatus() == ELauncherWorkerStatus::Canceling)
+			{
+				return LOCTEXT("DoneButtonCancellingTooltip", "Currently canceling.");
+			}
+		}
+		return LOCTEXT("DoneButtonCloseTooltip", "Close this page.");
+	}
+	FText GetDoneButtonText() const
+	{
+		ILauncherWorkerPtr LauncherWorkerPtr = LauncherWorker.Pin();
+
+		if (LauncherWorkerPtr.IsValid())
+		{
+			if (LauncherWorkerPtr->GetStatus() == ELauncherWorkerStatus::Busy)
+			{
+				return LOCTEXT("DoneButtonCancelLabel", "Cancel");
+			}
+			else if (LauncherWorkerPtr->GetStatus() == ELauncherWorkerStatus::Canceling)
+			{
+				return LOCTEXT("DoneButtonCancellingLabel", "Cancelling");
+			}
+		}
+		return LOCTEXT("DoneButtonDoneLabel", "Done");
+	}
+	
+
 	void ClearLog()
 	{
 		MessageList.Reset();
@@ -500,8 +610,11 @@ private:
 	// Holds the save button.
 	TSharedPtr<SButton> SaveButton;
 
-	// Holds a delegate to be invoked when this panel is closed
+	// Holds a delegate to be invoked when this panel is closed.
 	FOnClicked OnCloseClicked;
+
+	// Holds a delegate to be invoked when we want the launch profile rerun.
+	FOnClicked OnRerunClicked;
 };
 
 
diff --git a/Engine/Source/Developer/ProjectLauncher/Private/Widgets/SProjectLauncherSettings.h b/Engine/Source/Developer/ProjectLauncher/Private/Widgets/SProjectLauncherSettings.h
--- a/Engine/Source/Developer/ProjectLauncher/Private/Widgets/SProjectLauncherSettings.h
+++ b/Engine/Source/Developer/ProjectLauncher/Private/Widgets/SProjectLauncherSettings.h
@@ -126,6 +126,9 @@ private:
 	// Holds a delegate to be invoked when this panel is closed
 	FOnClicked OnCloseClicked;
 
+	// Holds a delegate to be invoked when a rerun of the profile is requested
+	FOnClicked OnRerunClicked;
+
 	// Holds a delegate to be invoked when this profile is deleted
 	FOnProfileRun OnDeleteClicked;
 
diff --git a/Engine/Source/Developer/SessionFrontend/Private/Widgets/Browser/SSessionBrowser.cpp b/Engine/Source/Developer/SessionFrontend/Private/Widgets/Browser/SSessionBrowser.cpp
--- a/Engine/Source/Developer/SessionFrontend/Private/Widgets/Browser/SSessionBrowser.cpp
+++ b/Engine/Source/Developer/SessionFrontend/Private/Widgets/Browser/SSessionBrowser.cpp
@@ -136,7 +136,7 @@ void SSessionBrowser::FilterSessions( )
 	{
 		const ISessionInfoPtr& Session = AvailableSessions[Index];
 
-		if (Session->GetSessionOwner() == FPlatformProcess::UserName(false))
+		if (Session->GetSessionOwner() == FPlatformProcess::UserName(true))
 		{
 			SessionList.Add(Session);
 		}
@@ -322,7 +322,7 @@ FReply SSessionBrowser::HandleTerminateSessionButtonClicked( )
 
 		if (SelectedSession.IsValid())
 		{
-			if (SelectedSession->GetSessionOwner() == FPlatformProcess::UserName(false))
+			if (SelectedSession->GetSessionOwner() == FPlatformProcess::UserName(true))
 			{
 				SelectedSession->Terminate();
 			}
diff --git a/Engine/Source/Developer/SessionServices/Private/SessionInstanceInfo.h b/Engine/Source/Developer/SessionServices/Private/SessionInstanceInfo.h
--- a/Engine/Source/Developer/SessionServices/Private/SessionInstanceInfo.h
+++ b/Engine/Source/Developer/SessionServices/Private/SessionInstanceInfo.h
@@ -92,7 +92,7 @@ public:
 	{
 		if (MessageEndpoint.IsValid() && EngineAddress.IsValid())
 		{
-			MessageEndpoint->Send(new FEngineServiceExecuteCommand(CommandString, FPlatformProcess::UserName(false)), EngineAddress);
+			MessageEndpoint->Send(new FEngineServiceExecuteCommand(CommandString, FPlatformProcess::UserName(true)), EngineAddress);
 		}
 	}
 
@@ -175,7 +175,7 @@ public:
 	{
 		if (MessageEndpoint.IsValid() && EngineAddress.IsValid())
 		{
-			MessageEndpoint->Send(new FEngineServiceTerminate(FPlatformProcess::UserName(false)), EngineAddress);
+			MessageEndpoint->Send(new FEngineServiceTerminate(FPlatformProcess::UserName(true)), EngineAddress);
 		}
 	}
 
diff --git a/Engine/Source/Developer/SessionServices/Private/SessionManager.cpp b/Engine/Source/Developer/SessionServices/Private/SessionManager.cpp
--- a/Engine/Source/Developer/SessionServices/Private/SessionManager.cpp
+++ b/Engine/Source/Developer/SessionServices/Private/SessionManager.cpp
@@ -170,7 +170,7 @@ void FSessionManager::FindExpiredSessions( const FDateTime& Now )
 
 bool FSessionManager::IsValidOwner( const FString& Owner )
 {
-	if (Owner == FPlatformProcess::UserName(false))
+	if (Owner == FPlatformProcess::UserName(true))
 	{
 		return true;
 	}
diff --git a/Engine/Source/Developer/TargetDeviceServices/Private/TargetDeviceProxy.cpp b/Engine/Source/Developer/TargetDeviceServices/Private/TargetDeviceProxy.cpp
--- a/Engine/Source/Developer/TargetDeviceServices/Private/TargetDeviceProxy.cpp
+++ b/Engine/Source/Developer/TargetDeviceServices/Private/TargetDeviceProxy.cpp
@@ -198,18 +198,18 @@ bool FTargetDeviceProxy::LaunchApp(FName InVariant, const FString& AppId, EBuild
 
 void FTargetDeviceProxy::PowerOff( bool Force )
 {
-	MessageEndpoint->Send(new FTargetDeviceServicePowerOff(FPlatformProcess::UserName(false), Force), MessageAddress);
+	MessageEndpoint->Send(new FTargetDeviceServicePowerOff(FPlatformProcess::UserName(true), Force), MessageAddress);
 }
 
 void FTargetDeviceProxy::PowerOn( )
 {
-	MessageEndpoint->Send(new FTargetDeviceServicePowerOn(FPlatformProcess::UserName(false)), MessageAddress);
+	MessageEndpoint->Send(new FTargetDeviceServicePowerOn(FPlatformProcess::UserName(true)), MessageAddress);
 }
 
 
 void FTargetDeviceProxy::Reboot( )
 {
-	MessageEndpoint->Send(new FTargetDeviceServiceReboot(FPlatformProcess::UserName(false)), MessageAddress);
+	MessageEndpoint->Send(new FTargetDeviceServiceReboot(FPlatformProcess::UserName(true)), MessageAddress);
 }
 
 void FTargetDeviceProxy::Run(FName InVariant, const FString& ExecutablePath, const FString& Params)
diff --git a/Engine/Source/Developer/TargetDeviceServices/Private/TargetDeviceProxyManager.cpp b/Engine/Source/Developer/TargetDeviceServices/Private/TargetDeviceProxyManager.cpp
--- a/Engine/Source/Developer/TargetDeviceServices/Private/TargetDeviceProxyManager.cpp
+++ b/Engine/Source/Developer/TargetDeviceServices/Private/TargetDeviceProxyManager.cpp
@@ -73,7 +73,7 @@ void FTargetDeviceProxyManager::GetProxies(FName TargetPlatformName, bool Includ
 	{
 		const TSharedPtr<FTargetDeviceProxy>& Proxy = It.Value();
 
-		if ((IncludeUnshared || Proxy->IsShared()) || (Proxy->GetHostUser() == FPlatformProcess::UserName(false)))
+		if ((IncludeUnshared || Proxy->IsShared()) || (Proxy->GetHostUser() == FPlatformProcess::UserName(true)))
 		{
 			if (TargetPlatformName == NAME_None || Proxy->HasTargetPlatform(TargetPlatformName))
 			{
@@ -106,7 +106,7 @@ void FTargetDeviceProxyManager::SendPing()
 {
 	if (MessageEndpoint.IsValid())
 	{
-		MessageEndpoint->Publish(new FTargetDeviceServicePing(FPlatformProcess::UserName(false)), EMessageScope::Network);
+		MessageEndpoint->Publish(new FTargetDeviceServicePing(FPlatformProcess::UserName(true)), EMessageScope::Network);
 	}
 }
 
diff --git a/Engine/Source/Developer/TargetDeviceServices/Private/TargetDeviceService.cpp b/Engine/Source/Developer/TargetDeviceServices/Private/TargetDeviceService.cpp
--- a/Engine/Source/Developer/TargetDeviceServices/Private/TargetDeviceService.cpp
+++ b/Engine/Source/Developer/TargetDeviceServices/Private/TargetDeviceService.cpp
@@ -80,7 +80,7 @@ void FTargetDeviceService::HandlePingMessage(const FTargetDeviceServicePing& InM
 		return;
 	}
 
-	if (Shared || (InMessage.HostUser == FPlatformProcess::UserName(false)))
+	if (Shared || (InMessage.HostUser == FPlatformProcess::UserName(true)))
 	{
 		ITargetDevicePtr DefaultDevice = GetDevice(); // Default Device is needed here!
 
diff --git a/Engine/Source/Developer/TargetDeviceServices/Private/TargetDeviceService.h b/Engine/Source/Developer/TargetDeviceServices/Private/TargetDeviceService.h
--- a/Engine/Source/Developer/TargetDeviceServices/Private/TargetDeviceService.h
+++ b/Engine/Source/Developer/TargetDeviceServices/Private/TargetDeviceService.h
@@ -151,7 +151,7 @@ public:
 			// notify other services
 			ClaimAddress = MessageEndpoint->GetAddress();
 			ClaimHost = FPlatformProcess::ComputerName();
-			ClaimUser = FPlatformProcess::UserName(false);
+			ClaimUser = FPlatformProcess::UserName(true);
 
 			MessageEndpoint->Publish(new FTargetDeviceClaimed(DeviceName, ClaimHost, ClaimUser));
 
@@ -166,7 +166,7 @@ public:
 		if (Running)
 		{
 			// notify other services
-			MessageEndpoint->Publish(new FTargetDeviceUnclaimed(DeviceName, FPlatformProcess::ComputerName(), FPlatformProcess::UserName(false)));
+			MessageEndpoint->Publish(new FTargetDeviceUnclaimed(DeviceName, FPlatformProcess::ComputerName(), FPlatformProcess::UserName(true)));
 
 			Running = false;
 		}
@@ -253,7 +253,7 @@ private:
 			{
 				if (Context->GetSender() != MessageEndpoint->GetAddress())
 				{
-					MessageEndpoint->Send(new FTargetDeviceClaimDenied(DeviceName, FPlatformProcess::ComputerName(), FPlatformProcess::UserName(false)), Context->GetSender());
+					MessageEndpoint->Send(new FTargetDeviceClaimDenied(DeviceName, FPlatformProcess::ComputerName(), FPlatformProcess::UserName(true)), Context->GetSender());
 				}
 			}
 			else
diff --git a/Engine/Source/Runtime/LauncherAutomatedService/Private/LauncherAutomatedServiceProvider.cpp b/Engine/Source/Runtime/LauncherAutomatedService/Private/LauncherAutomatedServiceProvider.cpp
--- a/Engine/Source/Runtime/LauncherAutomatedService/Private/LauncherAutomatedServiceProvider.cpp
+++ b/Engine/Source/Runtime/LauncherAutomatedService/Private/LauncherAutomatedServiceProvider.cpp
@@ -238,7 +238,7 @@ void FLauncherAutomatedServiceProvider::HandleDeviceProxyManagerProxyAdded( cons
 	if (AddedProxy->IsConnected())
 	{
 		FString SessionName = FString::Printf(TEXT("%s - %s"), FPlatformProcess::ComputerName(), *AutomatedDeviceGroup->GetName());
-		FString CommandLine = FString::Printf(TEXT("%s %s -SessionId=\"%s\" -SessionOwner=\"%s\" -SessionName=\"%s\""), *AutomatedProfile->GetDefaultLaunchRole()->GetInitialMap(), *AutomatedProfile->GetDefaultLaunchRole()->GetCommandLine(), *SessionID.ToString(), FPlatformProcess::UserName(false), *SessionName);
+		FString CommandLine = FString::Printf(TEXT("%s %s -SessionId=\"%s\" -SessionOwner=\"%s\" -SessionName=\"%s\""), *AutomatedProfile->GetDefaultLaunchRole()->GetInitialMap(), *AutomatedProfile->GetDefaultLaunchRole()->GetCommandLine(), *SessionID.ToString(), FPlatformProcess::UserName(true), *SessionName);
 
 		// @todo gmp: fix automated service provider; must use ILauncher interface here
 		//AddedProxy->Run(AutomatedProfile->GetBuildConfiguration(), AutomatedProfile->GetBuildGame(), CommandLine);
