Fixed a crash caused by a profile refering to an unsupported variant.

diff --git a/Engine/Source/Developer/ProjectLauncher/Private/Widgets/Deploy/SProjectLauncherSimpleDeviceListRow.h b/Engine/Source/Developer/ProjectLauncher/Private/Widgets/Deploy/SProjectLauncherSimpleDeviceListRow.h
--- a/Engine/Source/Developer/ProjectLauncher/Private/Widgets/Deploy/SProjectLauncherSimpleDeviceListRow.h
+++ b/Engine/Source/Developer/ProjectLauncher/Private/Widgets/Deploy/SProjectLauncherSimpleDeviceListRow.h
@@ -331,9 +331,20 @@ private:
 
 	void UpdateProfile()
 	{
-		LaunchProfile->SetDeploymentMode(ELauncherProfileDeploymentModes::FileServer);
+		check(DeviceProxy.IsValid());
 
 		FName Variant = SimpleProfile->GetDeviceVariant();
+		
+		// If the profile refers to a variant no longer supported by this device fall back to the default.
+		if (!DeviceProxy->HasVariant(Variant))
+		{
+			Variant = NAME_None;
+			SimpleProfile->SetDeviceVariant(Variant);
+		}
+
+		// Setup the Profile
+		LaunchProfile->SetDeploymentMode(ELauncherProfileDeploymentModes::FileServer);
+
 		ILauncherDeviceGroupRef NewGroup = Model->GetProfileManager()->CreateUnmanagedDeviceGroup();
 		NewGroup->AddDevice(DeviceProxy->GetTargetDeviceId(Variant));
 		LaunchProfile->SetDeployedDeviceGroup(NewGroup);
diff --git a/Engine/Source/Developer/TargetDeviceServices/Private/TargetDeviceProxy.cpp b/Engine/Source/Developer/TargetDeviceServices/Private/TargetDeviceProxy.cpp
--- a/Engine/Source/Developer/TargetDeviceServices/Private/TargetDeviceProxy.cpp
+++ b/Engine/Source/Developer/TargetDeviceServices/Private/TargetDeviceProxy.cpp
@@ -90,6 +90,11 @@ int32 FTargetDeviceProxy::GetVariants(TArray<FName>& OutVariants) const
 	return TargetDeviceVariants.GetKeys(OutVariants);
 }
 
+bool FTargetDeviceProxy::HasVariant(FName InVariant) const
+{
+	return TargetDeviceVariants.Contains(InVariant);
+}
+
 FName FTargetDeviceProxy::GetTargetDeviceVariant(const FString& InDeviceId) const
 {
 	for (TMap<FName, FTargetDeviceProxyVariant>::TConstIterator ItVariant(TargetDeviceVariants); ItVariant; ++ItVariant)
diff --git a/Engine/Source/Developer/TargetDeviceServices/Private/TargetDeviceProxy.h b/Engine/Source/Developer/TargetDeviceServices/Private/TargetDeviceProxy.h
--- a/Engine/Source/Developer/TargetDeviceServices/Private/TargetDeviceProxy.h
+++ b/Engine/Source/Developer/TargetDeviceServices/Private/TargetDeviceProxy.h
@@ -92,6 +92,8 @@ public:
 
 	virtual int32 GetVariants(TArray<FName>& OutVariants) const override;
 
+	virtual bool HasVariant(FName InVariant) const override;
+
 	virtual FName GetTargetDeviceVariant(const FString& InDeviceId) const override;
 
 	virtual bool HasDeviceId(const FString& InDeviceId) const override;
diff --git a/Engine/Source/Developer/TargetDeviceServices/Public/Interfaces/ITargetDeviceProxy.h b/Engine/Source/Developer/TargetDeviceServices/Public/Interfaces/ITargetDeviceProxy.h
--- a/Engine/Source/Developer/TargetDeviceServices/Public/Interfaces/ITargetDeviceProxy.h
+++ b/Engine/Source/Developer/TargetDeviceServices/Public/Interfaces/ITargetDeviceProxy.h
@@ -97,6 +97,14 @@ public:
 	 */
 	virtual int32 GetVariants(TArray<FName>& OutVariants) const = 0;
 
+	/**
+	 * Checks whether this device proxy contains a variant
+	 *
+	 * @param InVariant Variant to check for.
+	 * @return true is this variant is found.
+	 */
+	virtual bool HasVariant(FName InVariant) const = 0;
+
 	/**
 	 * Gets the variant name of the target device.
 	 *
