Code and documentation cleanup pass

diff --git a/Engine/Source/Developer/ProjectLauncher/Private/Widgets/Cook/SProjectLauncherCookByTheBookSettings.h b/Engine/Source/Developer/ProjectLauncher/Private/Widgets/Cook/SProjectLauncherCookByTheBookSettings.h
--- a/Engine/Source/Developer/ProjectLauncher/Private/Widgets/Cook/SProjectLauncherCookByTheBookSettings.h
+++ b/Engine/Source/Developer/ProjectLauncher/Private/Widgets/Cook/SProjectLauncherCookByTheBookSettings.h
@@ -1,9 +1,5 @@
 // Copyright 1998-2014 Epic Games, Inc. All Rights Reserved.
 
-/*=============================================================================
-	SProjectLauncherCookByTheBookSettings.h: Declares the SProjectLauncherCookByTheBookSettings class.
-=============================================================================*/
-
 #pragma once
 
 
@@ -11,14 +7,10 @@ namespace EShowMapsChoices
 {
 	enum Type
 	{
-		/**
-		 * Show all available maps.
-		 */
+		/** Show all available maps. */
 		ShowAllMaps,
 
-		/**
-		 * Only show maps that are to be cooked.
-		 */
+		/** Only show maps that are to be cooked. */
 		ShowCookedMaps
 	};
 }
@@ -35,15 +27,11 @@ public:
 	SLATE_BEGIN_ARGS(SProjectLauncherCookByTheBookSettings) { }
 	SLATE_END_ARGS()
 
-
 public:
 
-	/**
-	 * Destructor.
-	 */
+	/** Destructor. */
 	~SProjectLauncherCookByTheBookSettings( );
 
-
 public:
 
 	/**
@@ -54,108 +42,102 @@ public:
 	 */
 	void Construct(	const FArguments& InArgs, const FProjectLauncherModelRef& InModel, bool InShowSimple = false );
 
-
 protected:
 
-	/**
-	 * Refreshes the list of available maps.
-	 */
+	/** Refreshes the list of available maps. */
 	void RefreshMapList( );
 
-	/**
-	 * Refreshes the list of available cultures.
-	 */
+	/** Refreshes the list of available cultures. */
 	void RefreshCultureList( );
 
-
 private:
 
-	// Handles clicking the 'Select All Cultures' button.
+	/** Handles clicking the 'Select All Cultures' button. */
 	void HandleAllCulturesHyperlinkNavigate( bool AllPlatforms );
 
-	// Handles determining the visibility of the 'Select All Cultures' button.
+	/** Handles determining the visibility of the 'Select All Cultures' button. */
 	EVisibility HandleAllCulturesHyperlinkVisibility( ) const;
 
-	// Handles clicking the 'Select All Maps' button.
+	/** Handles clicking the 'Select All Maps' button. */
 	void HandleAllMapsHyperlinkNavigate( bool AllPlatforms );
 
-	// Handles selecting a build configuration for the cooker.
+	/** Handles selecting a build configuration for the cooker. */
 	void HandleCookConfigurationSelectorConfigurationSelected( EBuildConfigurations::Type );
 
-	// Handles getting the content text of the cooker build configuration selector.
+	/** Handles getting the content text of the cooker build configuration selector. */
 	FString HandleCookConfigurationSelectorText( ) const;
 
-	// Handles check state changes of the 'Incremental' check box.
+	/** Handles check state changes of the 'Incremental' check box. */
 	void HandleIncrementalCheckBoxCheckStateChanged( ESlateCheckBoxState::Type NewState );
 
-	// Handles determining the checked state of the 'Incremental' check box.
+	/** Handles determining the checked state of the 'Incremental' check box. */
 	ESlateCheckBoxState::Type HandleIncrementalCheckBoxIsChecked( ) const;
 
-	// Handles generating a row widget in the map list view.
+	/** Handles generating a row widget in the map list view. */
 	TSharedRef<ITableRow> HandleMapListViewGenerateRow( TSharedPtr<FString> InItem, const TSharedRef<STableViewBase>& OwnerTable );
 
-	// Handles generating a row widget in the culture list view.
+	/** Handles generating a row widget in the culture list view. */
 	TSharedRef<ITableRow> HandleCultureListViewGenerateRow( TSharedPtr<FString> InItem, const TSharedRef<STableViewBase>& OwnerTable );
 
-	// Handles determining the visibility of the 'Select All Maps' button.
+	/** Handles determining the visibility of the 'Select All Maps' button. */
 	EVisibility HandleMapSelectionHyperlinkVisibility( ) const;
 
-	// Handles getting the visibility of the map selection warning message.
+	/** Handles getting the visibility of the map selection warning message. */
 	EVisibility HandleNoMapSelectedBoxVisibility( ) const;
 
-	// Handles getting the text in the 'No maps' text block.
+	/** Handles getting the text in the 'No maps' text block. */
 	FText HandleNoMapsTextBlockText( ) const;
 
-	// Handles changing the selected profile in the profile manager.
+	/** Handles changing the selected profile in the profile manager. */
 	void HandleProfileManagerProfileSelected( const ILauncherProfilePtr& SelectedProfile, const ILauncherProfilePtr& PreviousProfile );
 
-	// Handles checking the specified 'Show maps' check box.
+	/** Handles checking the specified 'Show maps' check box. */
 	void HandleShowCheckBoxCheckStateChanged( ESlateCheckBoxState::Type NewState, EShowMapsChoices::Type Choice );
 
-	// Handles determining the checked state of the specified 'Show maps' check box.
+	/** Handles determining the checked state of the specified 'Show maps' check box. */
 	ESlateCheckBoxState::Type HandleShowCheckBoxIsChecked( EShowMapsChoices::Type Choice ) const;
 
-	// Handles checking the specified 'Unversioned' check box.
+	/** Handles checking the specified 'Unversioned' check box. */
 	void HandleUnversionedCheckBoxCheckStateChanged( ESlateCheckBoxState::Type NewState );
 
-	// Handles determining the checked state of the specified 'Unversioned' check box.
+	/** Handles determining the checked state of the specified 'Unversioned' check box. */
 	ESlateCheckBoxState::Type HandleUnversionedCheckBoxIsChecked( ) const;
 
-	// Handles determining the visibility of a validation error icon.
+	/** Handles determining the visibility of a validation error icon. */
 	EVisibility HandleValidationErrorIconVisibility( ELauncherProfileValidationErrors::Type Error ) const;
 
-	// Callback for getting the cookers additional options
+	/** Callback for getting the cookers additional options. */
 	FText HandleCookOptionsTextBlockText() const;
 
-	// Callback for changing the cookers additional options
+	/** Callback for changing the cookers additional options */
 	void HandleCookerOptionsCommitted(const FText& NewText, ETextCommit::Type CommitType);
 
-	// Callback for updating any settings after the selected project has changed in the profile
+	/** Callback for updating any settings after the selected project has changed in the profile. */
 	void HandleProfileProjectChanged();
 
-	// creates the complex widget
+	/** creates the complex widget. */
 	TSharedRef<SWidget> MakeComplexWidget();
 
-	// creates the simple widget
+	/** creates the simple widget. */
 	TSharedRef<SWidget> MakeSimpleWidget();
 
 private:
 
-	// Holds the culture list.
+	/** Holds the culture list. */
 	TArray<TSharedPtr<FString> > CultureList;
 
-	// Holds the map list view.
+	/** Holds the map list view. */
 	TSharedPtr<SListView<TSharedPtr<FString> > > CultureListView;
 
-	// Holds the map list.
+	/** Holds the map list. */
 	TArray<TSharedPtr<FString> > MapList;
 
-	// Holds the map list view.
+	/** Holds the map list view. */
 	TSharedPtr<SListView<TSharedPtr<FString> > > MapListView;
 
-	// Holds a pointer to the data model.
+	/** Holds a pointer to the data model. */
 	FProjectLauncherModelPtr Model;
 
-	// Holds the current 'Show maps' check box choice.
+	/** Holds the current 'Show maps' check box choice. */
 	EShowMapsChoices::Type ShowMapsChoice;
 };
diff --git a/Engine/Source/Developer/ProjectLauncher/Private/Widgets/SProjectLauncherSettings.h b/Engine/Source/Developer/ProjectLauncher/Private/Widgets/SProjectLauncherSettings.h
--- a/Engine/Source/Developer/ProjectLauncher/Private/Widgets/SProjectLauncherSettings.h
+++ b/Engine/Source/Developer/ProjectLauncher/Private/Widgets/SProjectLauncherSettings.h
@@ -1,9 +1,5 @@
 // Copyright 1998-2014 Epic Games, Inc. All Rights Reserved.
 
-/*=============================================================================
-	SProjectLauncherSettings.h: Declares the SProjectLauncherCookPage class.
-=============================================================================*/
-
 #pragma once
 
 
@@ -50,16 +46,12 @@ public:
 
 public:
 
-	/**
-	* Default constructor.
-	*/
+	/** Default constructor. */
 	SProjectLauncherSettings()
 		: CommandList(new FUICommandList())
 	{ }
 
-	/**
-	 * Destructor.
-	 */
+	/** Destructor. */
 	~SProjectLauncherSettings();
 
 public:
@@ -76,62 +68,62 @@ public:
 
 private:
 
-	// Create the menu command list
+	/** Create the menu command list. */
 	void CreateCommands();
 
-	// Make a Toolbar using the command list
+	/** Make a Toolbar using the command list. */
 	TSharedRef<SWidget> MakeToolbar(const TSharedRef<FUICommandList>& InCommandList);
 
 private:
 
-	// Callback for getting the launch profile we are currently editing.
+	/** Callback for getting the launch profile we are currently editing. */
 	ILauncherProfilePtr GetLaunchProfile()const;
 
-	// Callback for getting the visibility of the 'Select Profile' text block.
+	/** Callback for getting the visibility of the 'Select Profile' text block. */
 	EVisibility HandleSelectProfileTextBlockVisibility() const;
 
-	// Callback for getting the visibility of the settings scroll box.
+	/** Callback for getting the visibility of the settings scroll box. */
 	EVisibility HandleSettingsScrollBoxVisibility() const;
 
-	// Callback to get the name of the launch profile.
+	/** Callback to get the name of the launch profile. */
 	FText OnGetNameText() const;
 
-	// Callback to set the name of the launch profile.
+	/** Callback to set the name of the launch profile. */
 	void OnNameTextCommitted(const FText& NewText, ETextCommit::Type InTextCommit);
 
-	// Callback to get the description of the launch profile.
+	/** Callback to get the description of the launch profile. */
 	FText OnGetDescriptionText() const;
 
-	// Callback to set the description of the launch profile.
+	/** Callback to set the description of the launch profile. */
 	void OnDescriptionTextCommitted(const FText& NewText, ETextCommit::Type InTextCommit);
 
-	// Callbacks for executing the 'Close' action.
+	/** Callbacks for executing the 'Close' action. */
 	void HandleCloseActionExecute();
 	bool HandleCloseActionIsChecked() const;
 	bool HandleCloseActionCanExecute() const;
 
-	// Callbacks for executing the 'Delete' action.
+	/** Callbacks for executing the 'Delete' action. */
 	void HandleDeleteActionExecute();
 	bool HandleDeleteActionIsChecked() const;
 	bool HandleDeleteActionCanExecute() const;
 
 private:
 
-	// Holds a pointer to the data model.
+	/** Holds a pointer to the data model. */
 	FProjectLauncherModelPtr Model;
 
-	// Holds the list of UI commands for the profile settings.
+	/** Holds the list of UI commands for the profile settings. */
 	TSharedRef<FUICommandList> CommandList;
 
-	// Holds a delegate to be invoked when this panel is closed
+	/** Holds a delegate to be invoked when this panel is closed. */
 	FOnClicked OnCloseClicked;
 
-	// Holds a delegate to be invoked when a rerun of the profile is requested
+	/** Holds a delegate to be invoked when a rerun of the profile is requested. */
 	FOnClicked OnRerunClicked;
 
-	// Holds a delegate to be invoked when this profile is deleted
+	/** Holds a delegate to be invoked when this profile is deleted. */
 	FOnProfileRun OnDeleteClicked;
 
-	// Hold a pointer to the launch profile name edit box.
+	/** Hold a pointer to the launch profile name edit box. */
 	TSharedPtr<SInlineEditableTextBlock> NameEditBox;
 };
diff --git a/Engine/Source/Developer/ShaderCompilerCommon/Private/ShaderCompilerCommon.cpp b/Engine/Source/Developer/ShaderCompilerCommon/Private/ShaderCompilerCommon.cpp
--- a/Engine/Source/Developer/ShaderCompilerCommon/Private/ShaderCompilerCommon.cpp
+++ b/Engine/Source/Developer/ShaderCompilerCommon/Private/ShaderCompilerCommon.cpp
@@ -1,14 +1,12 @@
 // Copyright 1998-2014 Epic Games, Inc. All Rights Reserved.
 
-/*=============================================================================
-	ShaderCompilerCommon.cpp: Common functionality for platform Shader Compilers.
-=============================================================================*/
-
 #include "ShaderCompilerCommon.h"
 #include "ModuleManager.h"
 
+
 IMPLEMENT_MODULE(FDefaultModuleImpl, ShaderCompilerCommon);
 
+
 int16 GetNumUniformBuffersUsed(const FShaderResourceTable& InSRT)
 {
 	auto CountLambda = [&](const TArray<uint32>& In)
@@ -62,6 +60,7 @@ void BuildResourceTableTokenStream(const TArray<uint32>& InResourceMap, int32 Ma
 	}
 }
 
+
 void BuildResourceTableMapping(
 	const TMap<FString,FResourceTableEntry>& ResourceTableMap,
 	const TMap<FString,uint32>& ResourceTableLayoutHashes,
@@ -128,6 +127,7 @@ void BuildResourceTableMapping(
 	OutSRT.MaxBoundResourceTable = MaxBoundResourceTable;
 }
 
+
 bool RemoveUniformBuffersFromSource(FString& SourceCode)
 {
 	static const FString StaticStructToken(TEXT("static const struct"));
@@ -182,6 +182,7 @@ bool RemoveUniformBuffersFromSource(FString& SourceCode)
 	return true;
 }
 
+
 FString CreateCrossCompilerBatchFileContents(const FString& ShaderFile, const FString& OutputFile, const FString& FrequencySwitch, const FString& EntryPoint, const FString& VersionSwitch, const FString& ExtraArguments)
 {
 	FString BatchFile = TEXT("@echo off");
diff --git a/Engine/Source/Developer/ShaderCompilerCommon/Public/ShaderCompilerCommon.h b/Engine/Source/Developer/ShaderCompilerCommon/Public/ShaderCompilerCommon.h
--- a/Engine/Source/Developer/ShaderCompilerCommon/Public/ShaderCompilerCommon.h
+++ b/Engine/Source/Developer/ShaderCompilerCommon/Public/ShaderCompilerCommon.h
@@ -1,14 +1,11 @@
 // Copyright 1998-2014 Epic Games, Inc. All Rights Reserved.
 
-/*=============================================================================
-	ShaderCompilerCommon.h: Common functionality for platform Shader Compilers.
-=============================================================================*/
-
 #pragma once
 
 #include "Core.h"
 #include "ShaderCore.h"
 
+
 /**
  * This function looks for resources specified in ResourceTableMap in the 
  * parameter map, adds them to the resource table, and removes them from the
diff --git a/Engine/Source/Developer/TargetDeviceServices/Private/TargetDeviceServiceManager.cpp b/Engine/Source/Developer/TargetDeviceServices/Private/TargetDeviceServiceManager.cpp
--- a/Engine/Source/Developer/TargetDeviceServices/Private/TargetDeviceServiceManager.cpp
+++ b/Engine/Source/Developer/TargetDeviceServices/Private/TargetDeviceServiceManager.cpp
@@ -1,9 +1,5 @@
 // Copyright 1998-2014 Epic Games, Inc. All Rights Reserved.
 
-/*=============================================================================
-	TargetDeviceServiceManager.cpp: Implements the FTargetDeviceManager class.
-=============================================================================*/
-
 #include "TargetDeviceServicesPrivatePCH.h"
 
 
diff --git a/Engine/Source/Developer/TargetPlatform/Private/TargetPlatformManagerModule.cpp b/Engine/Source/Developer/TargetPlatform/Private/TargetPlatformManagerModule.cpp
--- a/Engine/Source/Developer/TargetPlatform/Private/TargetPlatformManagerModule.cpp
+++ b/Engine/Source/Developer/TargetPlatform/Private/TargetPlatformManagerModule.cpp
@@ -1,24 +1,23 @@
 // Copyright 1998-2014 Epic Games, Inc. All Rights Reserved.
 
-/*=============================================================================
-	TargetPlatformManagerModule.cpp: Implements the FTargetPlatformManagerModule class.
-=============================================================================*/
-
 #include "TargetPlatformPrivatePCH.h"
 #include "PlatformInfo.h"
 #include "DesktopPlatformModule.h"
 #include "IPhysXFormatModule.h"
 #include "IPhysXFormat.h"
 
+
 DEFINE_LOG_CATEGORY_STATIC(LogTargetPlatformManager, Log, All);
 
+
 // autosdks only function properly on windows right now.
 #if !IS_MONOLITHIC && (PLATFORM_WINDOWS)
-#define AUTOSDKS_ENABLED 1
+	#define AUTOSDKS_ENABLED 1
 #else
-#define AUTOSDKS_ENABLED 0
+	#define AUTOSDKS_ENABLED 0
 #endif
 
+
 /**
  * Module for the target platform manager
  */
@@ -27,13 +26,10 @@ class FTargetPlatformManagerModule
 {
 public:
 
-	/**
-	 * Default constructor.
-	 */
+	/** Default constructor. */
 	FTargetPlatformManagerModule()
 		: bRestrictFormatsToRuntimeOnly(false), bForceCacheUpdate(true)
 	{
-
 #if AUTOSDKS_ENABLED		
 		
 		// AutoSDKs only enabled if UE_SDKS_ROOT is set.
@@ -73,18 +69,18 @@ public:
 		bForceCacheUpdate = false;
 
 		FModuleManager::Get().OnModulesChanged().AddRaw(this, &FTargetPlatformManagerModule::ModulesChangesCallback);
-
-
 	}
 
+	/** Destructor. */
 	virtual ~FTargetPlatformManagerModule()
 	{
 		FModuleManager::Get().OnModulesChanged().RemoveAll(this);
 	}
 
-
 public:
 
+	// ITargetPlatformManagerModule interface
+
 	virtual void Invalidate()
 	{
 		bForceCacheUpdate = true;
@@ -113,12 +109,12 @@ public:
 	{
 		ITargetPlatform* Platform = FindTargetPlatform(DeviceId.GetPlatformName());
 
-		if (Platform != NULL)
+		if (Platform != nullptr)
 		{
 			return Platform->GetDevice(DeviceId);
 		}
 
-		return NULL;
+		return nullptr;
 	}
 
 	virtual ITargetPlatform* FindTargetPlatform(FString Name) override
@@ -133,7 +129,7 @@ public:
 			}
 		}
 
-		return NULL;
+		return nullptr;
 	}
 
 	virtual const TArray<ITargetPlatform*>& GetCookingTargetPlatforms() override
@@ -248,12 +244,12 @@ public:
 	virtual ITargetPlatform* GetRunningTargetPlatform() override
 	{
 		static bool bInitialized = false;
-		static ITargetPlatform* Result = NULL;
+		static ITargetPlatform* Result = nullptr;
 
 		if (!bInitialized || bForceCacheUpdate)
 		{
 			bInitialized = true;
-			Result = NULL;
+			Result = nullptr;
 
 			const TArray<ITargetPlatform*>& TargetPlatforms = GetTargetPlatforms();	
 
@@ -262,7 +258,7 @@ public:
 				if (TargetPlatforms[Index]->IsRunningPlatform())
 				{
 					 // we should not have two running platforms
-					checkf((Result == NULL),
+					checkf((Result == nullptr),
 						TEXT("Found multiple running platforms.\n\t%s\nand\n\t%s"),
 						*Result->PlatformName(),
 						*TargetPlatforms[Index]->PlatformName()
@@ -300,7 +296,7 @@ public:
 				if (Module)
 				{
 					IAudioFormat* Format = Module->GetAudioFormat();
-					if (Format != NULL)
+					if (Format != nullptr)
 					{
 						Results.Add(Format);
 					}
@@ -330,7 +326,7 @@ public:
 			}
 		}
 
-		return NULL;
+		return nullptr;
 	}
 
 	virtual const TArray<const ITextureFormat*>& GetTextureFormats() override
@@ -358,7 +354,7 @@ public:
 				if (Module)
 				{
 					ITextureFormat* Format = Module->GetTextureFormat();
-					if (Format != NULL)
+					if (Format != nullptr)
 					{
 						Results.Add(Format);
 					}
@@ -388,7 +384,7 @@ public:
 			}
 		}
 
-		return NULL;
+		return nullptr;
 	}
 
 	virtual const TArray<const IShaderFormat*>& GetShaderFormats() override
@@ -416,7 +412,7 @@ public:
 				if (Module)
 				{
 					IShaderFormat* Format = Module->GetShaderFormat();
-					if (Format != NULL)
+					if (Format != nullptr)
 					{
 						Results.Add(Format);
 					}
@@ -445,7 +441,7 @@ public:
 			}
 		}
 
-		return NULL;
+		return nullptr;
 	}
 
 	virtual uint16 ShaderFormatVersion(FName Name) override
@@ -492,7 +488,7 @@ public:
 				if (Module)
 				{
 					IPhysXFormat* Format = Module->GetPhysXFormat();
-					if (Format != NULL)
+					if (Format != nullptr)
 					{
 						Results.Add(Format);
 					}
@@ -522,12 +518,16 @@ public:
 			}
 		}
 
-		return NULL;
+		return nullptr;
 	}
 
-
 protected:
 
+	/**
+	 * Checks whether the AutoDesk software development kit (SDK) is enabled.
+	 *
+	 * @return true if the SDK is enabled, false otherwise.
+	 */
 	bool IsAutoSDKsEnabled()
 	{
 		static const FString SDKRootEnvFar(TEXT("UE_SDKS_ROOT"));
@@ -543,9 +543,7 @@ protected:
 		return false;
 	}
 
-	/**
-	 * Discovers the available target platforms.
-	 */
+	/** Discovers the available target platforms. */
 	void DiscoverAvailablePlatforms( )
 	{
 		Platforms.Empty(Platforms.Num());
@@ -568,7 +566,7 @@ protected:
 			if (Module)
 			{
 				ITargetPlatform* Platform = Module->GetTargetPlatform();
-				if (Platform != NULL)
+				if (Platform != nullptr)
 				{
 					// would like to move this check to GetActiveTargetPlatforms, but too many things cache this result
 					// this setup will become faster after TTP 341897 is complete.					
@@ -914,6 +912,8 @@ private:
 #endif
 };
 
+
 FString FTargetPlatformManagerModule::SDKStatusMessage = TEXT("");
 
+
 IMPLEMENT_MODULE(FTargetPlatformManagerModule, TargetPlatform);
diff --git a/Engine/Source/Developer/TargetPlatform/Public/Common/TargetPlatformBase.h b/Engine/Source/Developer/TargetPlatform/Public/Common/TargetPlatformBase.h
--- a/Engine/Source/Developer/TargetPlatform/Public/Common/TargetPlatformBase.h
+++ b/Engine/Source/Developer/TargetPlatform/Public/Common/TargetPlatformBase.h
@@ -1,13 +1,10 @@
 // Copyright 1998-2014 Epic Games, Inc. All Rights Reserved.
 
-/*=============================================================================
-	TTargetPlatformBase.h: Declares the TTargetPlatformBase template class.
-=============================================================================*/
-
 #pragma once
 
 #include "PlatformInfo.h"
 
+
 /**
  * Base class for target platforms.
  */
@@ -39,7 +36,6 @@ public:
 		OutFormats.Add(FName(TEXT("FullHDR")));
 	}
 
-
 #ifdef TEXTURERESOURCE_H_INCLUDED // defined in TextureResource.h, this way we know if UTexture is available, needed for Clang
 	FName GetDefaultTextureFormatName( const UTexture* Texture, const FConfigFile& EngineSettings ) const
 	{
@@ -217,7 +213,7 @@ protected:
 /**
  * Template for target platforms.
  *
- * @param TPlatformProperties - Type of platform properties.
+ * @param TPlatformProperties Type of platform properties.
  */
 template<typename TPlatformProperties>
 class TTargetPlatformBase
@@ -225,9 +221,7 @@ class TTargetPlatformBase
 {
 public:
 
-	/**
-	 * Default constructor.
-	 */
+	/** Default constructor. */
 	TTargetPlatformBase( )
 		: FTargetPlatformBase( PlatformInfo::FindPlatformInfo(TPlatformProperties::PlatformName()) )
 	{
@@ -237,7 +231,7 @@ public:
 
 public:
 
-	// Begin ITargetPlatform interface
+	// ITargetPlatform interface
 
 	virtual bool HasEditorOnlyData( ) const override
 	{
@@ -338,6 +332,4 @@ public:
 		return FName(TPlatformProperties::GetPhysicsFormat());
 	}
 #endif // WITH_ENGINE
-
-	// End ITargetPlatform interface
 };
diff --git a/Engine/Source/Developer/TargetPlatform/Public/Interfaces/IAudioFormat.h b/Engine/Source/Developer/TargetPlatform/Public/Interfaces/IAudioFormat.h
--- a/Engine/Source/Developer/TargetPlatform/Public/Interfaces/IAudioFormat.h
+++ b/Engine/Source/Developer/TargetPlatform/Public/Interfaces/IAudioFormat.h
@@ -1,9 +1,5 @@
 // Copyright 1998-2014 Epic Games, Inc. All Rights Reserved.
 
-/*=============================================================================
-	ITargetPlatform.h: Declares the ITargetPlatform interface.
-=============================================================================*/
-
 #pragma once
 
 
@@ -101,19 +97,16 @@ public:
 	virtual int32 Recompress( FName Format, const TArray<uint8>& SrcBuffer, FSoundQualityInfo& QualityInfo, TArray<uint8>& OutBuffer ) const = 0;
 
 	/**
-	 * Splits compressed data into chunks suitable for streaming audio
-	 *
-	 * @param	SrcBuffer	pre-compressed data as an array of bytes
-	 * @param	OutBuffers	Array of buffers that contain the chunks the original data was split into
+	 * Splits compressed data into chunks suitable for streaming audio.
 	 *
-	 * @return	bool	Whether bulk data could be split for streaming
+	 * @param SrcBuffer Pre-compressed data as an array of bytes.
+	 * @param OutBuffers Array of buffers that contain the chunks the original data was split into.
+	 * @return Whether bulk data could be split for streaming.
 	 */
 	virtual bool SplitDataForStreaming(const TArray<uint8>& SrcBuffer, TArray<TArray<uint8>>& OutBuffers) const {return false;}
 
 public:
 
-	/**
-	 * Virtual destructor.
-	 */
+	/** Virtual destructor. */
 	virtual ~IAudioFormat() { }
 };
diff --git a/Engine/Source/Developer/TargetPlatform/Public/Interfaces/IAudioFormatModule.h b/Engine/Source/Developer/TargetPlatform/Public/Interfaces/IAudioFormatModule.h
--- a/Engine/Source/Developer/TargetPlatform/Public/Interfaces/IAudioFormatModule.h
+++ b/Engine/Source/Developer/TargetPlatform/Public/Interfaces/IAudioFormatModule.h
@@ -1,9 +1,5 @@
 // Copyright 1998-2014 Epic Games, Inc. All Rights Reserved.
 
-/*=============================================================================
-	IAudioFormatModule.h: Declares the IAudioFormatModule interface.
-=============================================================================*/
-
 #pragma once
 
 
@@ -17,13 +13,13 @@ public:
 
 	/**
 	 * Gets the audio format.
+	 *
+	 * @return The audio format interface.
 	 */
 	virtual IAudioFormat* GetAudioFormat( ) = 0;
 
 public:
 
-	/**
-	 * Virtual destructor.
-	 */
+	/** Virtual destructor. */
 	~IAudioFormatModule( ) { }
 };
diff --git a/Engine/Source/Developer/TargetPlatform/Public/Interfaces/IShaderFormat.h b/Engine/Source/Developer/TargetPlatform/Public/Interfaces/IShaderFormat.h
--- a/Engine/Source/Developer/TargetPlatform/Public/Interfaces/IShaderFormat.h
+++ b/Engine/Source/Developer/TargetPlatform/Public/Interfaces/IShaderFormat.h
@@ -1,9 +1,5 @@
 // Copyright 1998-2014 Epic Games, Inc. All Rights Reserved.
 
-/*=============================================================================
-	ITextureFormatModule.h: Declares the ITextureFormatModule interface.
-=============================================================================*/
-
 #pragma once
 
 
@@ -45,8 +41,6 @@ public:
 
 public:
 
-	/**
-	 * Virtual destructor.
-	 */
+	/** Virtual destructor. */
 	virtual ~IShaderFormat( ) { }
 };
diff --git a/Engine/Source/Developer/TargetPlatform/Public/Interfaces/IShaderFormatModule.h b/Engine/Source/Developer/TargetPlatform/Public/Interfaces/IShaderFormatModule.h
--- a/Engine/Source/Developer/TargetPlatform/Public/Interfaces/IShaderFormatModule.h
+++ b/Engine/Source/Developer/TargetPlatform/Public/Interfaces/IShaderFormatModule.h
@@ -1,9 +1,5 @@
 // Copyright 1998-2014 Epic Games, Inc. All Rights Reserved.
 
-/*=============================================================================
-	IShaderFormatModule.h: Declares the IShaderFormatModule interface.
-=============================================================================*/
-
 #pragma once
 
 
@@ -17,13 +13,13 @@ public:
 
 	/**
 	 * Gets the shader format.
+	 *
+	 * @return The shader format interface.
 	 */
 	virtual IShaderFormat* GetShaderFormat( ) = 0;
 
 public:
 
-	/**
-	 * Virtual destructor.
-	 */
+	/** Virtual destructor. */
 	~IShaderFormatModule( ) { }
 };
diff --git a/Engine/Source/Developer/TargetPlatform/Public/Interfaces/ITargetDevice.h b/Engine/Source/Developer/TargetPlatform/Public/Interfaces/ITargetDevice.h
--- a/Engine/Source/Developer/TargetPlatform/Public/Interfaces/ITargetDevice.h
+++ b/Engine/Source/Developer/TargetPlatform/Public/Interfaces/ITargetDevice.h
@@ -1,9 +1,5 @@
 // Copyright 1998-2014 Epic Games, Inc. All Rights Reserved.
 
-/*=============================================================================
-	ITargetDevice.h: Declares the ITargetDevice interface.
-=============================================================================*/
-
 #pragma once
 
 
@@ -381,8 +377,6 @@ public:
 
 public:
 	
-	/**
-	 * Virtual destructor.
-	 */
+	/** Virtual destructor. */
 	virtual ~ITargetDevice( ) { }
 };
diff --git a/Engine/Source/Developer/TargetPlatform/Public/Interfaces/ITargetPlatform.h b/Engine/Source/Developer/TargetPlatform/Public/Interfaces/ITargetPlatform.h
--- a/Engine/Source/Developer/TargetPlatform/Public/Interfaces/ITargetPlatform.h
+++ b/Engine/Source/Developer/TargetPlatform/Public/Interfaces/ITargetPlatform.h
@@ -1,9 +1,5 @@
 // Copyright 1998-2014 Epic Games, Inc. All Rights Reserved.
 
-/*=============================================================================
-	ITargetPlatform.h: Declares the ITargetPlatform interface.
-=============================================================================*/
-
 #pragma once
 
 
@@ -154,7 +150,7 @@ public:
 	virtual const PlatformInfo::FPlatformInfo& GetPlatformInfo( ) const = 0;
 
 	/**
-	 * Gets the platform's ini name (so an offline tool can load the ini for the given target platform)
+	 * Gets the platform's INI name (so an offline tool can load the INI for the given target platform).
 	 *
 	 * @see PlatformName
 	 */
@@ -182,9 +178,9 @@ public:
 	/** 
 	 * Generates a platform specific asset manifest given an array of FAssetData.
 	 *
-	 * @param ChunkMap A map of asset path to ChunkIDs for all of the assets
-	 * @param ChunkIDsInUse A set of all ChunkIDs used by this set of assets
-	 * @return true if the manifest was successfully generated, or if the platform doesn't need a manifest 
+	 * @param ChunkMap A map of asset path to ChunkIDs for all of the assets.
+	 * @param ChunkIDsInUse A set of all ChunkIDs used by this set of assets.
+	 * @return true if the manifest was successfully generated, or if the platform doesn't need a manifest .
 	 */
 	virtual bool GenerateStreamingInstallManifest( const TMultiMap<FString, int32>& ChunkMap, const TSet<int32>& ChunkIDsInUse ) const = 0;
 
@@ -267,7 +263,7 @@ public:
 	/**
 	 * Checks whether this platform supports the specified build target, i.e. Game or Editor.
 	 *
-	 * @param BuildTarget - The build target to check.
+	 * @param BuildTarget The build target to check.
 	 * @return true if the build target is supported, false otherwise.
 	 */
 	virtual bool SupportsBuildTarget( EBuildTargets::Type BuildTarget ) const = 0;
@@ -280,7 +276,7 @@ public:
 	/**
 	 * Checks whether the target platform supports the specified feature.
 	 *
-	 * @param Feature - The feature to check.
+	 * @param Feature The feature to check.
 	 * @return true if the feature is supported, false otherwise.
 	 */
 	virtual bool SupportsFeature( ETargetPlatformFeatures::Type Feature ) const = 0;
@@ -420,8 +416,6 @@ public:
 
 public:
 
-	/**
-	 * Virtual destructor.
-	 */
+	/** Virtual destructor. */
 	virtual ~ITargetPlatform() { }
 };
diff --git a/Engine/Source/Developer/TargetPlatform/Public/Interfaces/ITargetPlatformManagerModule.h b/Engine/Source/Developer/TargetPlatform/Public/Interfaces/ITargetPlatformManagerModule.h
--- a/Engine/Source/Developer/TargetPlatform/Public/Interfaces/ITargetPlatformManagerModule.h
+++ b/Engine/Source/Developer/TargetPlatform/Public/Interfaces/ITargetPlatformManagerModule.h
@@ -1,9 +1,5 @@
 // Copyright 1998-2014 Epic Games, Inc. All Rights Reserved.
 
-/*=============================================================================
-	ITargetPlatformManagerModule.h: Declares the ITargetPlatformManagerModule interface.
-=============================================================================*/
-
 #pragma once
 
 
@@ -63,7 +59,6 @@ public:
 	 */
 	virtual ITargetPlatform* FindTargetPlatform( FString Name ) = 0;
 
-
 	/**
 	 * Return the list of platforms which we need to support when cooking (only set when actually cooking)
 	 *
@@ -147,8 +142,6 @@ public:
 
 public:
 
-	/**
-	 * Virtual destructor.
-	 */
+	/** Virtual destructor. */
 	~ITargetPlatformManagerModule() { }
 };
diff --git a/Engine/Source/Developer/TargetPlatform/Public/Interfaces/ITargetPlatformModule.h b/Engine/Source/Developer/TargetPlatform/Public/Interfaces/ITargetPlatformModule.h
--- a/Engine/Source/Developer/TargetPlatform/Public/Interfaces/ITargetPlatformModule.h
+++ b/Engine/Source/Developer/TargetPlatform/Public/Interfaces/ITargetPlatformModule.h
@@ -1,9 +1,5 @@
 // Copyright 1998-2014 Epic Games, Inc. All Rights Reserved.
 
-/*=============================================================================
-	ITargetPlatformModule.h: Declares the ITargetPlatformModule interface.
-=============================================================================*/
-
 #pragma once
 
 
@@ -22,8 +18,6 @@ public:
 
 public:
 
-	/**
-	 * Virtual destructor.
-	 */
+	/** Virtual destructor. */
 	~ITargetPlatformModule( ) { }
 };
diff --git a/Engine/Source/Developer/TargetPlatform/Public/Interfaces/ITextureFormat.h b/Engine/Source/Developer/TargetPlatform/Public/Interfaces/ITextureFormat.h
--- a/Engine/Source/Developer/TargetPlatform/Public/Interfaces/ITextureFormat.h
+++ b/Engine/Source/Developer/TargetPlatform/Public/Interfaces/ITextureFormat.h
@@ -1,9 +1,5 @@
 // Copyright 1998-2014 Epic Games, Inc. All Rights Reserved.
 
-/*=============================================================================
-	ITextureFormat.h: Declares the ITextureFormat interface.
-=============================================================================*/
-
 #pragma once
 
 
@@ -60,25 +56,24 @@ public:
 	/**
 	 * Gets the list of supported formats.
 	 *
-	 * @param OutFormats - Will hold the list of formats.
+	 * @param OutFormats Will hold the list of formats.
 	 */
 	virtual void GetSupportedFormats( TArray<FName>& OutFormats ) const = 0;
 
 	/**
 	* Gets the capabilities of the texture compressor.
 	*
-	* @param OutCaps - filled with capability properties of texture format compressor.
+	* @param OutCaps Filled with capability properties of texture format compressor.
 	*/
 	virtual FTextureFormatCompressorCaps GetFormatCapabilities( ) const = 0;
 
 	/**
 	 * Compresses a single image.
 	 *
-	 * @param Image - The input image.
-	 * @param BuildSettings - Build settings.
-	 * @param bImageHasAlphaChannel - true if the image has a non-white alpha channel.
-	 * @param OutCompressedMip - The compressed image.
-	 *
+	 * @param Image The input image.
+	 * @param BuildSettings Build settings.
+	 * @param bImageHasAlphaChannel true if the image has a non-white alpha channel.
+	 * @param OutCompressedMip The compressed image.
 	 * @returns true on success, false otherwise.
 	 */
 	virtual bool CompressImage(
@@ -91,8 +86,6 @@ public:
 
 public:
 
-	/**
-	 * Virtual destructor.
-	 */
+	/** Virtual destructor. */
 	virtual ~ITextureFormat( ) { }
 };
diff --git a/Engine/Source/Developer/TargetPlatform/Public/Interfaces/ITextureFormatModule.h b/Engine/Source/Developer/TargetPlatform/Public/Interfaces/ITextureFormatModule.h
--- a/Engine/Source/Developer/TargetPlatform/Public/Interfaces/ITextureFormatModule.h
+++ b/Engine/Source/Developer/TargetPlatform/Public/Interfaces/ITextureFormatModule.h
@@ -1,9 +1,5 @@
 // Copyright 1998-2014 Epic Games, Inc. All Rights Reserved.
 
-/*=============================================================================
-	ITextureFormatModule.h: Declares the ITextureFormatModule interface.
-=============================================================================*/
-
 #pragma once
 
 
@@ -17,13 +13,13 @@ public:
 
 	/**
 	 * Gets the texture format.
+	 *
+	 * @return The texture format interface.
 	 */
 	virtual ITextureFormat* GetTextureFormat( ) = 0;
 
 public:
 
-	/**
-	 * Virtual destructor.
-	 */
+	/** Virtual destructor. */
 	~ITextureFormatModule() { }
 };
diff --git a/Engine/Source/Developer/TargetPlatform/Public/Interfaces/TargetDeviceId.h b/Engine/Source/Developer/TargetPlatform/Public/Interfaces/TargetDeviceId.h
--- a/Engine/Source/Developer/TargetPlatform/Public/Interfaces/TargetDeviceId.h
+++ b/Engine/Source/Developer/TargetPlatform/Public/Interfaces/TargetDeviceId.h
@@ -1,9 +1,5 @@
 // Copyright 1998-2014 Epic Games, Inc. All Rights Reserved.
 
-/*=============================================================================
-	TargetDeviceId.h: Declares the FRemoteDeviceUri class.
-=============================================================================*/
-
 #pragma once
 
 
@@ -14,9 +10,7 @@ class FTargetDeviceId
 {
 public:
 
-	/**
-	 * Default constructor.
-	 */
+	/** Default constructor. */
 	FTargetDeviceId( ) { }
 
 	/**
@@ -137,9 +131,9 @@ public:
 
 private:
 
-	// Holds the name of the device's target platform.
+	/** Holds the name of the device's target platform. */
 	FString PlatformName;
 
-	// Holds the name of the target device.
+	/** Holds the name of the target device. */
 	FString DeviceName;
 };
diff --git a/Engine/Source/Developer/TargetPlatform/Public/TargetPlatform.h b/Engine/Source/Developer/TargetPlatform/Public/TargetPlatform.h
--- a/Engine/Source/Developer/TargetPlatform/Public/TargetPlatform.h
+++ b/Engine/Source/Developer/TargetPlatform/Public/TargetPlatform.h
@@ -1,19 +1,15 @@
 // Copyright 1998-2014 Epic Games, Inc. All Rights Reserved.
 
-/*=============================================================================
-	TargetPlatform.h: TargetPlatform module public header file.
-=============================================================================*/
-
 #pragma once
 
 
-/* Dependencies
+/* Public dependencies
  *****************************************************************************/
 
 #include "ModuleManager.h"
 
 
-/* Interfaces
+/* Public interfaces
  *****************************************************************************/
 
 #include "TargetDeviceId.h"
diff --git a/Engine/Source/Editor/BlueprintGraph/Private/EdGraphSchema_K2.cpp b/Engine/Source/Editor/BlueprintGraph/Private/EdGraphSchema_K2.cpp
--- a/Engine/Source/Editor/BlueprintGraph/Private/EdGraphSchema_K2.cpp
+++ b/Engine/Source/Editor/BlueprintGraph/Private/EdGraphSchema_K2.cpp
@@ -1,10 +1,5 @@
 // Copyright 1998-2014 Epic Games, Inc. All Rights Reserved.
 
-/*=============================================================================
-	EdGraphSchema_K2.cpp
-=============================================================================*/
-
-
 #include "BlueprintGraphPrivatePCH.h"
 
 #include "Engine/LevelScriptBlueprint.h"
@@ -36,6 +31,7 @@
 #include "K2Node_Tunnel.h"
 #include "K2Node_SetFieldsInStruct.h"
 
+
 //////////////////////////////////////////////////////////////////////////
 // FBlueprintMetadata
 
diff --git a/Engine/Source/Editor/BlueprintGraph/Private/EdGraphSchema_K2_Actions.cpp b/Engine/Source/Editor/BlueprintGraph/Private/EdGraphSchema_K2_Actions.cpp
--- a/Engine/Source/Editor/BlueprintGraph/Private/EdGraphSchema_K2_Actions.cpp
+++ b/Engine/Source/Editor/BlueprintGraph/Private/EdGraphSchema_K2_Actions.cpp
@@ -1,9 +1,5 @@
 // Copyright 1998-2014 Epic Games, Inc. All Rights Reserved.
 
-/*=============================================================================
-	EdGraphSchema_K2.cpp
-=============================================================================*/
-
 #include "BlueprintGraphPrivatePCH.h"
 
 #include "GraphEditorActions.h"
@@ -13,8 +9,10 @@
 #include "Editor/UnrealEd/Public/EdGraphUtilities.h"
 #include "EdGraph/EdGraphNode_Documentation.h"
 
+
 #define SNAP_GRID (16) // @todo ensure this is the same as SNodePanel::GetSnapGridSize()
 
+
 namespace 
 {
 	// Maximum distance a drag can be off a node edge to require 'push off' from node
@@ -27,6 +25,7 @@ namespace
 	const float NodeLiteralHeight = 48.0f;
 }
 
+
 /////////////////////////////////////////////////////
 // FK2SchemaActionUtils
 
diff --git a/Engine/Source/Editor/DeviceProfileEditor/Private/DetailsPanel/SDeviceProfileDetailsPanel.cpp b/Engine/Source/Editor/DeviceProfileEditor/Private/DetailsPanel/SDeviceProfileDetailsPanel.cpp
--- a/Engine/Source/Editor/DeviceProfileEditor/Private/DetailsPanel/SDeviceProfileDetailsPanel.cpp
+++ b/Engine/Source/Editor/DeviceProfileEditor/Private/DetailsPanel/SDeviceProfileDetailsPanel.cpp
@@ -1,11 +1,5 @@
 // Copyright 1998-2014 Epic Games, Inc. All Rights Reserved.
 
-
-/*=============================================================================
-	SDeviceProfileDetailsPanel.cpp: Implements the SDeviceProfileDetailsPanel class.
-=============================================================================*/
-
-
 #include "DeviceProfileEditorPCH.h"
 #include "DeviceProfiles/DeviceProfile.h"
 
diff --git a/Engine/Source/Editor/DeviceProfileEditor/Private/DetailsPanel/SDeviceProfileDetailsPanel.h b/Engine/Source/Editor/DeviceProfileEditor/Private/DetailsPanel/SDeviceProfileDetailsPanel.h
--- a/Engine/Source/Editor/DeviceProfileEditor/Private/DetailsPanel/SDeviceProfileDetailsPanel.h
+++ b/Engine/Source/Editor/DeviceProfileEditor/Private/DetailsPanel/SDeviceProfileDetailsPanel.h
@@ -1,14 +1,10 @@
 // Copyright 1998-2014 Epic Games, Inc. All Rights Reserved.
 
-
-/*=============================================================================
-	SDeviceProfileDetailsPanel.h: Declares the SDeviceProfileDetailsPanel class.
-=============================================================================*/
-
 #pragma once
 
 #include "PropertyEditing.h"
 
+
 class UDeviceProfile;
 
 #define LOCTEXT_NAMESPACE "DeviceProfileDetailsPanel"
@@ -25,7 +21,6 @@ public:
 		SLATE_DEFAULT_SLOT( FArguments, Content )
 	SLATE_END_ARGS()
 
-
 	/**
 	 * Constructs this widget with InArgs
 	 */
@@ -39,24 +34,20 @@ public:
 	 */
 	void UpdateUIForProfile( const TWeakObjectPtr< UDeviceProfile > InProfile );
 
-
 private:
 
-	/**
-	 * Refresh the UI of the details panel.
-	 */
+	/** Refresh the UI of the details panel. */
 	void RefreshUI();
 
-
 private:
 
-	// The profile this panel is showing details for
+	/** The profile this panel is showing details for. */
 	TWeakObjectPtr< UDeviceProfile > ViewingProfile;
 
-	// The widget which hosts the details content if a profile is provided
+	/** The widget which hosts the details content if a profile is provided. */
 	TSharedPtr< SVerticalBox > DetailsViewBox;
 
-	// Holds the details view.
+	/** Holds the details view. */
 	TSharedPtr<IDetailsView> SettingsView;
 };
 
diff --git a/Engine/Source/Editor/DeviceProfileEditor/Private/DeviceProfileConsoleVariableColumn.cpp b/Engine/Source/Editor/DeviceProfileEditor/Private/DeviceProfileConsoleVariableColumn.cpp
--- a/Engine/Source/Editor/DeviceProfileEditor/Private/DeviceProfileConsoleVariableColumn.cpp
+++ b/Engine/Source/Editor/DeviceProfileEditor/Private/DeviceProfileConsoleVariableColumn.cpp
@@ -120,11 +120,6 @@ TSharedRef<class SWidget> FConsoleVariableCellPresenter::ConstructDisplayWidget(
 }
 
 
-/*=============================================================================
-	FDeviceProfileConsoleVariableColumn Implementation
-=============================================================================*/
-
-
 FDeviceProfileConsoleVariableColumn::FDeviceProfileConsoleVariableColumn()
 {
 }
diff --git a/Engine/Source/Editor/DeviceProfileEditor/Private/DeviceProfileConsoleVariableColumn.h b/Engine/Source/Editor/DeviceProfileEditor/Private/DeviceProfileConsoleVariableColumn.h
--- a/Engine/Source/Editor/DeviceProfileEditor/Private/DeviceProfileConsoleVariableColumn.h
+++ b/Engine/Source/Editor/DeviceProfileEditor/Private/DeviceProfileConsoleVariableColumn.h
@@ -1,12 +1,7 @@
 // Copyright 1998-2014 Epic Games, Inc. All Rights Reserved.
 
-/*=============================================================================
-	DeviceProfileConsoleVariableColumn.h: Declares the FDeviceProfileConsoleVariableColumn class.
-=============================================================================*/
-
 #pragma once
 
-
 #include "IPropertyTableCustomColumn.h"
 
 
@@ -23,16 +18,20 @@ class FDeviceProfileConsoleVariableColumn : public IPropertyTableCustomColumn
 public:
 	FDeviceProfileConsoleVariableColumn();
 
-	/** Begin IPropertyTableCustomColumn interface */
+public:
+
+	// IPropertyTableCustomColumn interface
+
 	virtual bool Supports( const TSharedRef< IPropertyTableColumn >& Column, const TSharedRef< IPropertyTableUtilities >& Utilities ) const override;
 	virtual TSharedPtr< SWidget > CreateColumnLabel( const TSharedRef< IPropertyTableColumn >& Column, const TSharedRef< IPropertyTableUtilities >& Utilities, const FName& Style ) const override;
 	virtual TSharedPtr< IPropertyTableCellPresenter > CreateCellPresenter( const TSharedRef< IPropertyTableCell >& Cell, const TSharedRef< IPropertyTableUtilities >& Utilities, const FName& Style ) const override;
-	/** End IPropertyTableCustomColumn interface */
+
+public:
 
 	/**
-	 * Delegate used to notify listeners that an edit request was triggered from the property table
+	 * Delegate used to notify listeners that an edit request was triggered from the property table.
 	 *
-	 * @return - Access to the delegate
+	 * @return The delegate.
 	 */
 	FOnEditDeviceProfileCVarsRequestDelegate& OnEditCVarsRequest()
 	{
@@ -41,7 +40,7 @@ public:
 
 private:
 
-	/** Delegate triggered when user opts to edit CVars **/
+	/** Delegate triggered when user opts to edit CVars. **/
 	FOnEditDeviceProfileCVarsRequestDelegate OnEditCVarsRequestDelegate;
 };
 
diff --git a/Engine/Source/Editor/DeviceProfileEditor/Private/DeviceProfileEditorModule.cpp b/Engine/Source/Editor/DeviceProfileEditor/Private/DeviceProfileEditorModule.cpp
--- a/Engine/Source/Editor/DeviceProfileEditor/Private/DeviceProfileEditorModule.cpp
+++ b/Engine/Source/Editor/DeviceProfileEditor/Private/DeviceProfileEditorModule.cpp
@@ -1,19 +1,13 @@
 // Copyright 1998-2014 Epic Games, Inc. All Rights Reserved.
 
-/*=============================================================================
-	DeviceProfileEditorModule.cpp: Implements the FDeviceProfileEditorModule class.
-=============================================================================*/
-
 #include "DeviceProfileEditorPCH.h"
 #include "WorkspaceMenuStructureModule.h"
 
 
 #define LOCTEXT_NAMESPACE "DeviceProfileEditor"
 
-
 IMPLEMENT_MODULE(FDeviceProfileEditorModule, DeviceProfileEditor);
 
-
 static const FName DeviceProfileEditorName("DeviceProfileEditor");
 
 
@@ -45,4 +39,4 @@ TSharedRef<SDockTab> FDeviceProfileEditorModule::SpawnDeviceProfileEditorTab(con
 }
 
 
-#undef LOCTEXT_NAMESPACE
\ No newline at end of file
+#undef LOCTEXT_NAMESPACE
diff --git a/Engine/Source/Editor/DeviceProfileEditor/Private/DeviceProfileEditorPCH.h b/Engine/Source/Editor/DeviceProfileEditor/Private/DeviceProfileEditorPCH.h
--- a/Engine/Source/Editor/DeviceProfileEditor/Private/DeviceProfileEditorPCH.h
+++ b/Engine/Source/Editor/DeviceProfileEditor/Private/DeviceProfileEditorPCH.h
@@ -1,31 +1,24 @@
 // Copyright 1998-2014 Epic Games, Inc. All Rights Reserved.
 
-
-/*=============================================================================
-	DeviceProfileEditorPCH.h: Pre-compiled header file for the DeviceProfileEditor module.
-=============================================================================*/
-
-
-#ifndef __DeviceProfileEditorPCH_h__
-#define __DeviceProfileEditorPCH_h__
+#pragma once
 
 
 /* Dependencies
  *****************************************************************************/
+
 #include "UnrealEd.h"
 #include "IPropertyTable.h"
 #include "TargetPlatform.h"
 #include "PlatformInfo.h"
 
+
 /* Private includes
 *****************************************************************************/
-#include "DetailsPanel/SDeviceProfileDetailsPanel.h"
 
+#include "DetailsPanel/SDeviceProfileDetailsPanel.h"
 #include "DeviceProfileConsoleVariableColumn.h"
 #include "SDeviceProfileEditorSingleProfileView.h"
 #include "SDeviceProfileCreateProfilePanel.h"
 #include "SDeviceProfileSelectionPanel.h"
 #include "SDeviceProfileEditor.h"
 #include "DeviceProfileEditor.h"
-
-#endif // __DeviceProfileEditorPCH_h__
diff --git a/Engine/Source/Editor/DeviceProfileEditor/Private/SDeviceProfileCreateProfilePanel.cpp b/Engine/Source/Editor/DeviceProfileEditor/Private/SDeviceProfileCreateProfilePanel.cpp
--- a/Engine/Source/Editor/DeviceProfileEditor/Private/SDeviceProfileCreateProfilePanel.cpp
+++ b/Engine/Source/Editor/DeviceProfileEditor/Private/SDeviceProfileCreateProfilePanel.cpp
@@ -1,10 +1,5 @@
 // Copyright 1998-2014 Epic Games, Inc. All Rights Reserved.
 
-/*=============================================================================
-	SDeviceProfileCreateProfilePanel.cpp: Implements the SDeviceProfileCreateProfilePanel class.
-=============================================================================*/
-
-
 #include "DeviceProfileEditorPCH.h"
 #include "DeviceProfiles/DeviceProfileManager.h"
 #include "DeviceProfiles/DeviceProfile.h"
@@ -14,12 +9,14 @@
 
 #define LOCTEXT_NAMESPACE "DeviceProfileCreateProfilePanel"
 
+
 /** Panel layout constants */
 namespace DeviceProfileCreateProfileUIConstants
 {
 	const FMargin ListElementPadding(10.0f, 2.0f);
 }
 
+
 BEGIN_SLATE_FUNCTION_BUILD_OPTIMIZATION
 void SDeviceProfileCreateProfilePanel::Construct( const FArguments& InArgs, TWeakObjectPtr< UDeviceProfileManager > InDeviceProfileManager )
 {
diff --git a/Engine/Source/Editor/DeviceProfileEditor/Private/SDeviceProfileCreateProfilePanel.h b/Engine/Source/Editor/DeviceProfileEditor/Private/SDeviceProfileCreateProfilePanel.h
--- a/Engine/Source/Editor/DeviceProfileEditor/Private/SDeviceProfileCreateProfilePanel.h
+++ b/Engine/Source/Editor/DeviceProfileEditor/Private/SDeviceProfileCreateProfilePanel.h
@@ -1,12 +1,8 @@
 // Copyright 1998-2014 Epic Games, Inc. All Rights Reserved.
 
-
-/*=============================================================================
-	SDeviceProfileCreateProfilePanel.h: Declares the SDeviceProfileCreateProfilePanel class.
-=============================================================================*/
-
 #pragma once
 
+
 #define LOCTEXT_NAMESPACE "DeviceProfileCreateProfilePanel"
 
 
@@ -21,10 +17,9 @@ public:
 		SLATE_DEFAULT_SLOT( FArguments, Content )
 	SLATE_END_ARGS()
 
-	/** Constructs this widget with InArgs */
+	/** Constructs this widget with InArgs. */
 	void Construct( const FArguments& InArgs, TWeakObjectPtr< UDeviceProfileManager > InDeviceProfileManager );
 
-
 public:
 
 	/**
@@ -41,48 +36,45 @@ public:
 	 */
 	FReply HandleCreateDeviceProfileButtonClicked();
 
-
 public:
 	
 	/**
 	 * Is the select a base profile selection available
 	 *
-	 * @return True if enabled, false otherwise
+	 * @return true if enabled, false otherwise.
 	 */
 	bool IsBaseProfileComboBoxEnabled() const;
 
 	/**
 	 * Handle base selection changed.
 	 *
-	 * @param NewSelection -  The newly selected device profile to use as a parent
-	 * @param SelectionInfo - The selection type
+	 * @param NewSelection The newly selected device profile to use as a parent.
+	 * @param SelectionInfo The selection type.
 	 */
 	void HandleBaseProfileSelectionChanged( UDeviceProfile* NewSelection, ESelectInfo::Type SelectInfo );
 
 	/**
 	 * Set the base combo box content.
 	 *
-	 * @return The device profile parent name
+	 * @return The device profile parent name.
 	 */
 	FString SetBaseProfileComboBoxContent() const;
 
 	/**
 	 * Handle base combo box generate widget.
 	 *
-	 * @param InItem - The device profile
-	 *
-	 * @return A text box holding the device profile name
+	 * @param InItem The device profile.
+	 * @return A text box holding the device profile name.
 	 */
 	TSharedRef<SWidget> HandleBaseComboBoxGenerateWidget( UDeviceProfile* InItem );
 
-
 public:
 	
 	/**
 	 * Handle profile type changed.
 	 *
-	 * @param NewSelection -  The newly selected device profile type
-	 * @param SelectionInfo - The selection type
+	 * @param NewSelection The newly selected device profile type.
+	 * @param SelectionInfoThe selection type.
 	 */
 	void HandleProfileTypeChanged( TSharedPtr<FString> NewSelection, ESelectInfo::Type SelectInfo );
 	
@@ -96,54 +88,49 @@ public:
 	/**
 	 * Handle profile type combo box generate widget.
 	 *
-	 * @param The type of this in string format
-	 *
-	 * @return A text box holding the device profile name
+	 * @param The type of this in string format.
+	 * @return A text box holding the device profile name.
 	 */
 	TSharedRef<SWidget> HandleProfileTypeComboBoxGenarateWidget( TSharedPtr<FString> InItem );
 
-
 private:
 
-	/**
-	 * Reset the components state.
-	 */
+	/** Reset the components state. */
 	void ResetComponentsState();
 
-
 private:
 
-	// Holds the Device Profile Class types combo box
+	/** Holds the Device Profile Class types combo box. */
 	TSharedPtr<SComboBox<TSharedPtr<FString>>> DeviceProfileTypesComboBox;
 	
-	// The collection of available types you can create a profile for
+	/** The collection of available types you can create a profile for. */
 	TArray<TSharedPtr<FString>> DeviceProfileTypes;
 
-	// Holds the Device Profile Class types combo box
+	/** Holds the Device Profile Class types combo box. */
 	TSharedPtr<SComboBox<UDeviceProfile*>> ParentObjectComboBox;
 
-	// The collection of device profiles available for you to copy properties from that matches the profile type
+	/** The collection of device profiles available for you to copy properties from that matches the profile type. */
 	TArray<UDeviceProfile*> AvailableBaseObjects;
 
-	// Holds the device profile combo box
+	/** Holds the device profile combo box. */
 	TSharedPtr< SComboBox< UDeviceProfile* > > DeviceProfileBaseComboBox;
 
-	// Holds the device profile check box
+	/** Holds the device profile check box. */
 	TSharedPtr< SCheckBox > DeviceProfileBaseCheckBox;
 
-	// Holds the device profile manager
+	/** Holds the device profile manager. */
 	TWeakObjectPtr< UDeviceProfileManager > DeviceProfileManager;
 
-	// Holds the device profile name text box
+	/** Holds the device profile name text box. */
 	TSharedPtr< SEditableTextBox > DeviceProfileNameTextBox;
 
-	// Holds the device profile name text box
+	/** Holds the device profile name text box. */
 	TSharedPtr< SEditableTextBox > DeviceProfileParentNameTextBox;
 
-	// Holds the selected device profile parent
+	/** Holds the selected device profile parent. */
 	TWeakObjectPtr< UDeviceProfile > SelectedDeviceProfileParent;
 
-	// Holds the selected device profile parent
+	/** Holds the selected device profile parent. */
 	TSharedPtr< FString > SelectedDeviceProfileType;
 };
 
diff --git a/Engine/Source/Editor/DeviceProfileEditor/Private/SDeviceProfileEditor.cpp b/Engine/Source/Editor/DeviceProfileEditor/Private/SDeviceProfileEditor.cpp
--- a/Engine/Source/Editor/DeviceProfileEditor/Private/SDeviceProfileEditor.cpp
+++ b/Engine/Source/Editor/DeviceProfileEditor/Private/SDeviceProfileEditor.cpp
@@ -1,11 +1,5 @@
 // Copyright 1998-2014 Epic Games, Inc. All Rights Reserved.
 
-
-/*=============================================================================
-	SDeviceProfileEditor.cpp: Implements the SDeviceProfileEditor class.
-=============================================================================*/
-
-
 #include "DeviceProfileEditorPCH.h"
 
 #include "DeviceProfiles/DeviceProfileManager.h"
@@ -21,11 +15,10 @@
 static const FName DeviceProfileEditorTabName("DeviceProfiles");
 
 
-
 /** Source control for the default device profile config saves */
-class SDeviceProfileSourceControl : public SCompoundWidget
+class SDeviceProfileSourceControl
+	: public SCompoundWidget
 {
-
 public:
 
 	SLATE_BEGIN_ARGS(SDeviceProfileSourceControl) {}
@@ -33,14 +26,12 @@ public:
 	SLATE_END_ARGS()
 
 
-	/** Constructs this widget with InArgs */
+	/** Constructs this widget with InArgs. */
 	void Construct(const FArguments& InArgs);
 
-	
-	/** Destructor */
+	/** Destructor. */
 	~SDeviceProfileSourceControl();
 
-
 	/**
 	 * Indicate which SWidgetSwitcher slot should be used to show the user of the source control status
 	 */
@@ -49,45 +40,44 @@ public:
 		return bIsDefaultConfigCheckOutNeeded ? 1 : 0;
 	}
 
-
 	/**
-	 * Take action to check out the default device profile configuration file when requested
+	 * Take action to check out the default device profile configuration file when requested.
 	 *
-	 * @return - Whether we handled the event
+	 * @return Whether we handled the event.
 	 */
 	FReply HandleSaveDefaultsButtonPressed();
 
-
 	/**
-	* Take action to check out the default device profile configuration file when requested
+	* Take action to check out the default device profile configuration file when requested.
 	*
-	* @return - Whether we handled the event
+	* @return Whether we handled the event.
 	*/
 	FReply HandleCheckoutButtonPressed();
 
 
 	/**
 	 * Check whether the SCC is enabled for the Checkout button to become available.
+	 *
+	 * @return true if SCC is enabled, false otherwise.
 	 */
 	bool IsCheckOutAvailable() const;
 
 
 public:
 
-	// Begin SCompoundWidget interface
-	virtual void Tick(const FGeometry& AllottedGeometry, const double InCurrentTime, const float InDeltaTime) override;
-	// End SCompoundWidget interface
+	// SCompoundWidget interface
 
+	virtual void Tick(const FGeometry& AllottedGeometry, const double InCurrentTime, const float InDeltaTime) override;
 
 private:
 
-	// Holds a timer for checking whether the device profile configuration file needs to be checked out.
+	/** Holds a timer for checking whether the device profile configuration file needs to be checked out. */
 	float DefaultConfigCheckOutTimer;
 
-	// Holds a flag indicating whether the section's configuration file needs to be checked out.
+	/** Holds a flag indicating whether the section's configuration file needs to be checked out. */
 	bool bIsDefaultConfigCheckOutNeeded;
 
-	// The direct path to the default device profile config file
+	/** The direct path to the default device profile config file. */
 	FString AbsoluteConfigFilePath;
 };
 
@@ -557,4 +547,4 @@ void SDeviceProfileEditor::RebuildPropertyTable()
 }
 
 
-#undef LOCTEXT_NAMESPACE
\ No newline at end of file
+#undef LOCTEXT_NAMESPACE
diff --git a/Engine/Source/Editor/DeviceProfileEditor/Private/SDeviceProfileEditor.h b/Engine/Source/Editor/DeviceProfileEditor/Private/SDeviceProfileEditor.h
--- a/Engine/Source/Editor/DeviceProfileEditor/Private/SDeviceProfileEditor.h
+++ b/Engine/Source/Editor/DeviceProfileEditor/Private/SDeviceProfileEditor.h
@@ -1,9 +1,5 @@
 // Copyright 1998-2014 Epic Games, Inc. All Rights Reserved.
 
-/*=============================================================================
-	SDeviceProfileEditor.h: Declares the SDeviceProfileEditor class.
-=============================================================================*/
-
 #pragma once
 
 
@@ -13,7 +9,8 @@
 /**
  * Slate widget to allow users to edit Device Profiles
  */
-class SDeviceProfileEditor : public SCompoundWidget
+class SDeviceProfileEditor
+	: public SCompoundWidget
 {
 public:
 
@@ -29,14 +26,13 @@ public:
 	/** Destructor */
 	~SDeviceProfileEditor();
 
-
 public:
 
 	/**
 	 * Callback for spawning tabs.
 	 *
-	 * @param Args - The arguments used for the tab that is spawned
-	 * @param TabIdentifier - The ID of the tab we want to spawn.
+	 * @param Args The arguments used for the tab that is spawned.
+	 * @param TabIdentifier The ID of the tab we want to spawn.
 	 */
 	TSharedRef<SDockTab> HandleTabManagerSpawnTab( const FSpawnTabArgs& Args, FName TabIdentifier );
 
@@ -44,18 +40,17 @@ public:
 	/**
 	* Callback for spawning tabs.
 	*
-	* @param Args - The arguments used for the tab that is spawned
-	* @param InDeviceProfile - The device profile the single editor is displaying information for.
+	* @param Args The arguments used for the tab that is spawned.
+	* @param InDeviceProfile The device profile the single editor is displaying information for.
 	*/
 	TSharedRef<SDockTab> HandleTabManagerSpawnSingleProfileTab(const FSpawnTabArgs& Args, TWeakObjectPtr<UDeviceProfile> InDeviceProfile);
 
-
 public:
 
 	/**
 	 * Handle a device being pinned to the grid.
 	 *
-	 * @param DeviceProfile - The selected profile.
+	 * @param DeviceProfile The selected profile.
 	 */
 	void HandleDeviceProfilePinned( const TWeakObjectPtr< UDeviceProfile >& DeviceProfile );
 
@@ -63,77 +58,73 @@ public:
 	/**
 	 * Handle the device being unpinned from the grid.
 	 *
-	 * @param DeviceProfile - The deselected profile.
+	 * @param DeviceProfile The deselected profile.
 	 */
 	void HandleDeviceProfileUnpinned( const TWeakObjectPtr< UDeviceProfile >& DeviceProfile );
 
 
 	/**
-	* Handle the device being unpinned from the grid.
-	*
-	* @param DeviceProfile - The profile that is to be viewed alone.
-	*/
+	 * Handle the device being unpinned from the grid.
+	 *
+	 * @param DeviceProfile The profile that is to be viewed alone.
+	 */
 	void HandleDeviceProfileViewAlone(const TWeakObjectPtr< UDeviceProfile >& DeviceProfile);
 
 
 	/**
 	 * Get whether the notification which indicates that no profiles currently being viewed in the grid, should be visible.
 	 *
-	 *	@return Whether the Notification is visible or not.
+	 * @return Whether the Notification is visible or not.
 	 */
 	EVisibility GetEmptyDeviceProfileGridNotificationVisibility() const;
 
 
 	/**
-	 * Setup the device profile property grid
+	 * Setup the device profile property grid.
 	 *
-	 * @return The widget containing the property grid
+	 * @return The widget containing the property grid.
 	 */
 	TSharedRef< SWidget > SetupPropertyEditor();
 
-
 private:
 
 	/**
 	 * Create the device profile panel which hosts the functionality to edit device profiles in the editor.
 	 *
-	 * @return The slate widget which hosts the device profile manipulator
+	 * @return The slate widget which hosts the device profile manipulator.
 	 */
 	TSharedPtr< SWidget > CreateMainDeviceProfilePanel();
 	
 
-	/**
-	 * Rebuild the device profile property table
-	 */
+	/** Rebuild the device profile property table. */
 	void RebuildPropertyTable();
 
-
 private:
 
-	// Hold a reference to the device profile manager
+	/** Hold a reference to the device profile manager. */
 	TWeakObjectPtr<UDeviceProfileManager> DeviceProfileManager;
 
-	// The collection of device profile ptrs for the selection process
+	/** The collection of device profile ptrs for the selection process. */
 	TArray<UObject*> DeviceProfiles;
 
-	// Holds the tab manager that manages the front-end's tabs.
+	/** Holds the tab manager that manages the front-end's tabs. */
 	TSharedPtr<FTabManager> TabManager;
 
-	// The widget which allows the user to select profiles
+	/** The widget which allows the user to select profiles. */
 	TSharedPtr<SDeviceProfileSelectionPanel> DeviceProfileSelectionPanel;
 
-	// Holds the property table
+	/** Holds the property table. */
 	TSharedPtr<IPropertyTable> PropertyTable;
 
-	// The list of Single Profile Editor ID's which have already been spawned/registered
+	/** The list of Single Profile Editor ID's which have already been spawned/registered. */
 	TArray<FName> RegisteredTabIds;
 
-	// Access to the "Window" menu to add our spawned tabs of single profile editors.
+	/** Access to the "Window" menu to add our spawned tabs of single profile editors. */
 	TSharedPtr<FWorkspaceItem> DeviceManagerMenuGroup;
 
-	// Holds the tab stack where single profiles and the device profile editor will be spawned to
+	/** Holds the tab stack where single profiles and the device profile editor will be spawned to. */
 	TSharedPtr<FTabManager::FStack> EditorTabStack;
 };
 
 
-#undef LOCTEXT_NAMESPACE
\ No newline at end of file
+#undef LOCTEXT_NAMESPACE
diff --git a/Engine/Source/Editor/DeviceProfileEditor/Private/SDeviceProfileEditorSingleProfileView.cpp b/Engine/Source/Editor/DeviceProfileEditor/Private/SDeviceProfileEditorSingleProfileView.cpp
--- a/Engine/Source/Editor/DeviceProfileEditor/Private/SDeviceProfileEditorSingleProfileView.cpp
+++ b/Engine/Source/Editor/DeviceProfileEditor/Private/SDeviceProfileEditorSingleProfileView.cpp
@@ -1,10 +1,5 @@
 // Copyright 1998-2014 Epic Games, Inc. All Rights Reserved.
 
-/*=============================================================================
-	SDeviceProfileEditorSingleProfileView.cpp: Implements the SDeviceProfileEditorSingleProfileView class.
-=============================================================================*/
-
-
 #include "DeviceProfileEditorPCH.h"
 #include "DeviceProfiles/DeviceProfile.h"
 #include "PropertyEditorModule.h"
@@ -43,4 +38,4 @@ void  SDeviceProfileEditorSingleProfileView::Construct(const FArguments& InArgs,
 }
 
 
-#undef LOCTEXT_NAMESPACE
\ No newline at end of file
+#undef LOCTEXT_NAMESPACE
diff --git a/Engine/Source/Editor/DeviceProfileEditor/Private/SDeviceProfileEditorSingleProfileView.h b/Engine/Source/Editor/DeviceProfileEditor/Private/SDeviceProfileEditorSingleProfileView.h
--- a/Engine/Source/Editor/DeviceProfileEditor/Private/SDeviceProfileEditorSingleProfileView.h
+++ b/Engine/Source/Editor/DeviceProfileEditor/Private/SDeviceProfileEditorSingleProfileView.h
@@ -1,10 +1,5 @@
 // Copyright 1998-2014 Epic Games, Inc. All Rights Reserved.
 
-
-/*=============================================================================
-	SDeviceProfileEditorSingleProfileView.h: Declares the SDeviceProfileEditorSingleProfileView class.
-=============================================================================*/
-
 #pragma once
 
 
@@ -14,7 +9,8 @@
 /**
  * Slate widget to allow users to select device profiles
  */
-class  SDeviceProfileEditorSingleProfileView : public SCompoundWidget
+class  SDeviceProfileEditorSingleProfileView
+	: public SCompoundWidget
 {
 public:
 
@@ -22,22 +18,20 @@ public:
 		SLATE_DEFAULT_SLOT( FArguments, Content )
 	SLATE_END_ARGS()
 
-
 	/** Constructs this widget with InArgs */
 	void Construct( const FArguments& InArgs, TWeakObjectPtr< UDeviceProfile > InDeviceProfile );
 
-
 	/** Destructor */
 	~SDeviceProfileEditorSingleProfileView(){}
 
 private:
 
-	// The profile selected from the current list
+	/** The profile selected from the current list. */
 	TWeakObjectPtr< UDeviceProfile > EditingProfile;
 
-	// Holds the details view.
+	/** Holds the details view. */
 	TSharedPtr<IDetailsView> SettingsView;
 };
 
 
-#undef LOCTEXT_NAMESPACE
\ No newline at end of file
+#undef LOCTEXT_NAMESPACE
diff --git a/Engine/Source/Editor/DeviceProfileEditor/Private/SDeviceProfileSelectionPanel.cpp b/Engine/Source/Editor/DeviceProfileEditor/Private/SDeviceProfileSelectionPanel.cpp
--- a/Engine/Source/Editor/DeviceProfileEditor/Private/SDeviceProfileSelectionPanel.cpp
+++ b/Engine/Source/Editor/DeviceProfileEditor/Private/SDeviceProfileSelectionPanel.cpp
@@ -1,10 +1,5 @@
 // Copyright 1998-2014 Epic Games, Inc. All Rights Reserved.
 
-/*=============================================================================
-	SDeviceProfileSelectionPanel.cpp: Implements the SDeviceProfileSelectionPanel class.
-=============================================================================*/
-
-
 #include "DeviceProfileEditorPCH.h"
 #include "DeviceProfiles/DeviceProfile.h"
 #include "DeviceProfiles/DeviceProfileManager.h"
@@ -17,7 +12,8 @@
 /**
  * Slate widget for each selection row in the selection list
  */
-class SDeviceProfileSelectionRow : public SMultiColumnTableRow< TWeakObjectPtr< UDeviceProfile > >
+class SDeviceProfileSelectionRow
+	: public SMultiColumnTableRow<TWeakObjectPtr<UDeviceProfile>>
 {
 public:
 
@@ -41,18 +37,17 @@ public:
 	/**
 	 * Generates the widget for the specified column.
 	 *
-	 * @param ColumnName - The name of the column to generate the widget for.
+	 * @param ColumnName The name of the column to generate the widget for.
 	 * @return The widget.
 	 */
 	virtual TSharedRef<SWidget> GenerateWidgetForColumn( const FName& ColumnName ) override;
 
-
 public:
 
 	/**
-	 * Handle a state change in the device profile pinning
+	 * Handle a state change in the device profile pinning.
 	 *
-	 * @return Whether the event was handled
+	 * @return Whether the event was handled.
 	 */
 	FReply HandleDeviceProfilePinStateChanged();
 
@@ -60,14 +55,14 @@ public:
 	/**
 	* Handle the view single profile button pressed.
 	*
-	* @return Whether the event was handled
+	* @return Whether the event was handled.
 	*/
 	FReply ViewSingleProfile();
 
 	/**
-	 * Get the image for the pin of this item
+	 * Get the image for the pin of this item.
 	 *
-	 * @return If the profile is pinned we want to indicate this by providing a different angle for the pin
+	 * @return If the profile is pinned we want to indicate this by providing a different angle for the pin.
 	 */
 	const FSlateBrush* GetPinnedImage() const;
 
@@ -75,40 +70,37 @@ public:
 	/**
 	 * Get the display profile name.
 	 *
-	 * @return The profile display name
+	 * @return The profile display name.
 	 */
 	FString GetProfileDisplayName() const;
 
-
 private:
 
-	// Holds the selected device profile
+	/** Holds the selected device profile. */
 	TWeakObjectPtr< UDeviceProfile > SelectedDeviceProfile;
 
-	// Delegate executed when a profile is unpinned
+	/** Delegate executed when a profile is unpinned. */
 	FOnDeviceProfileUnpinned OnDeviceProfileUnpinned;
 
-	// Delegate executed when a profile is pinned
+	/** Delegate executed when a profile is pinned. */
 	FOnDeviceProfilePinned OnDeviceProfilePinned;
 
-	// Delegate executed when requesting that a profile be viewed alone.
+	/** Delegate executed when requesting that a profile be viewed alone. */
 	FOnDeviceProfileViewAlone OnDeviceProfileViewAlone;
 
-	// A reference to the profiles pin button
+	/** A reference to the profiles pin button. */
 	TSharedPtr< SButton > PinProfileButton;
 
-	// A reference to the profiles view button
+	/** A reference to the profiles view button. */
 	TSharedPtr< SButton > ViewProfileButton;
 
-	// Whether this profile selection is pinned.
+	/** Whether this profile selection is pinned. */
 	bool bIsPinned;
 };
 
 
-/*=============================================================================
-	SDeviceProfileSelectionRow Implementation
-=============================================================================*/
-
+/* SDeviceProfileSelectionRow implementation
+ *****************************************************************************/
 
 void SDeviceProfileSelectionRow::Construct( const FArguments& InArgs, const TSharedRef< STableViewBase >& InOwnerTableView )
 {
@@ -213,10 +205,8 @@ FString SDeviceProfileSelectionRow::GetProfileDisplayName() const
 }
 
 
-/*=============================================================================
-	SDeviceProfileSelectionPanel Implementation
-=============================================================================*/
-
+/* SDeviceProfileSelectionPanel implementation
+ *****************************************************************************/
 
 void SDeviceProfileSelectionPanel::Construct( const FArguments& InArgs, TWeakObjectPtr< UDeviceProfileManager > InDeviceProfileManager )
 {
@@ -341,4 +331,4 @@ void SDeviceProfileSelectionPanel::RegenerateProfileList()
 }
 
 
-#undef LOCTEXT_NAMESPACE
\ No newline at end of file
+#undef LOCTEXT_NAMESPACE
diff --git a/Engine/Source/Editor/DeviceProfileEditor/Private/SDeviceProfileSelectionPanel.h b/Engine/Source/Editor/DeviceProfileEditor/Private/SDeviceProfileSelectionPanel.h
--- a/Engine/Source/Editor/DeviceProfileEditor/Private/SDeviceProfileSelectionPanel.h
+++ b/Engine/Source/Editor/DeviceProfileEditor/Private/SDeviceProfileSelectionPanel.h
@@ -1,10 +1,5 @@
 // Copyright 1998-2014 Epic Games, Inc. All Rights Reserved.
 
-
-/*=============================================================================
-	SDeviceProfileSelectionPanel.h: Declares the SDeviceProfileSelectionPanel class.
-=============================================================================*/
-
 #pragma once
 
 
@@ -24,7 +19,8 @@ DECLARE_DELEGATE_OneParam( FOnDeviceProfileViewAlone, const TWeakObjectPtr< UDev
 /**
  * Slate widget to allow users to select device profiles
  */
-class SDeviceProfileSelectionPanel : public SCompoundWidget
+class SDeviceProfileSelectionPanel
+	: public SCompoundWidget
 {
 public:
 
@@ -41,56 +37,49 @@ public:
 	SLATE_END_ARGS()
 
 
-	/** Constructs this widget with InArgs */
+	/** Constructs this widget with InArgs. */
 	void Construct( const FArguments& InArgs, TWeakObjectPtr< UDeviceProfileManager > InDeviceProfileManager );
 
-
-	/** Destructor */
+	/** Destructor. */
 	~SDeviceProfileSelectionPanel();
 
-
 	/**
 	 * Handle generating the device profile widget.
 	 *
-	 * @param InDeviceProfile - The selected device profile.
-	 * @param OwnerTable - The Owner table.
-	 *
-	 * @return The generated widget
+	 * @param InDeviceProfile The selected device profile.
+	 * @param OwnerTable The Owner table.
+	 * @return The generated widget.
 	 */
 	TSharedRef<ITableRow> OnGenerateWidgetForDeviceProfile( TWeakObjectPtr<UDeviceProfile> InDeviceProfile, const TSharedRef< STableViewBase >& OwnerTable );
 
-
 protected:
 
-	/**
-	* Regenerate the list view widget when the device profiles are recreated
-	*/
+	/** Regenerate the list view widget when the device profiles are recreated */
 	void RegenerateProfileList();
 
-
 private:
 
-	// Holds the device profile manager
+	/** Holds the device profile manager. */
 	TWeakObjectPtr< UDeviceProfileManager > DeviceProfileManager;
 
-	// The collection of device profile ptrs for the selection process
+	/** The collection of device profiles for the selection process. */
 	TArray< TWeakObjectPtr<UDeviceProfile> > DeviceProfiles;
 
-	// Hold the widget that contains the list view of device profiles
+	/** Hold the widget that contains the list view of device profiles. */
 	TSharedPtr< SVerticalBox > ListWidget;
 
-	// Delegate for handling a profile being pinned to the grid
+	/** Delegate for handling a profile being pinned to the grid. */
 	FOnDeviceProfilePinned OnDeviceProfilePinned;
 
-	// Delegate for handling a profile being unpinned from the grid
+	/** Delegate for handling a profile being unpinned from the grid. */
 	FOnDeviceProfileUnpinned OnDeviceProfileUnpinned;
 
-	// Delegate for handling a request to view the profile in it's own editor
+	/** Delegate for handling a request to view the profile in it's own editor. */
 	FOnDeviceProfileViewAlone OnDeviceProfileViewAlone;
 
-	// The profile selected from the current list
+	/** The profile selected from the current list. */
 	TWeakObjectPtr< UDeviceProfile > SelectedProfile;
 };
 
 
-#undef LOCTEXT_NAMESPACE
\ No newline at end of file
+#undef LOCTEXT_NAMESPACE
diff --git a/Engine/Source/Editor/DeviceProfileServices/Private/DeviceProfileServicesModule.cpp b/Engine/Source/Editor/DeviceProfileServices/Private/DeviceProfileServicesModule.cpp
--- a/Engine/Source/Editor/DeviceProfileServices/Private/DeviceProfileServicesModule.cpp
+++ b/Engine/Source/Editor/DeviceProfileServices/Private/DeviceProfileServicesModule.cpp
@@ -1,11 +1,8 @@
 // Copyright 1998-2014 Epic Games, Inc. All Rights Reserved.
 
-/*=============================================================================
-	DeviceProfileServicesModule.cpp: Implements the FDeviceProfileServicesModule module.
-=============================================================================*/
-
 #include "DeviceProfileServicesPCH.h"
 
+
 /**
  * Implements the DeviceProfileServices module.
  */
@@ -14,21 +11,22 @@ class FDeviceProfileServicesModule
 {
 public:
 
-		// Begin IDeviceProfileServicesModule interface
-		virtual IDeviceProfileServicesUIManagerRef GetProfileServicesManager( ) override
-		{
-			if (!DeviceProfileServicesUIManagerSingleton.IsValid())
-			{
-				DeviceProfileServicesUIManagerSingleton = MakeShareable(new FDeviceProfileServicesUIManager());
-			}
+	// IDeviceProfileServicesModule interface
 
-			return DeviceProfileServicesUIManagerSingleton.ToSharedRef();
+	virtual IDeviceProfileServicesUIManagerRef GetProfileServicesManager( ) override
+	{
+		if (!DeviceProfileServicesUIManagerSingleton.IsValid())
+		{
+			DeviceProfileServicesUIManagerSingleton = MakeShareable(new FDeviceProfileServicesUIManager());
 		}
-		// End IDeviceProfileServicesModule interface
+
+		return DeviceProfileServicesUIManagerSingleton.ToSharedRef();
+	}
 
 protected:
-		// Holds the session manager singleton.
-		static IDeviceProfileServicesUIManagerPtr DeviceProfileServicesUIManagerSingleton;
+
+	// Holds the session manager singleton.
+	static IDeviceProfileServicesUIManagerPtr DeviceProfileServicesUIManagerSingleton;
 };
 
 
diff --git a/Engine/Source/Editor/DeviceProfileServices/Private/DeviceProfileServicesPCH.h b/Engine/Source/Editor/DeviceProfileServices/Private/DeviceProfileServicesPCH.h
--- a/Engine/Source/Editor/DeviceProfileServices/Private/DeviceProfileServicesPCH.h
+++ b/Engine/Source/Editor/DeviceProfileServices/Private/DeviceProfileServicesPCH.h
@@ -1,11 +1,8 @@
 // Copyright 1998-2014 Epic Games, Inc. All Rights Reserved.
 
-/*=============================================================================
-	DeviceProfileServicesPCH.h: Pre-compiled header file for the DeviceProfileServices module.
-=============================================================================*/
-
 #pragma once
 
+
 #include "DeviceProfileServices.h"
 
 
diff --git a/Engine/Source/Editor/DeviceProfileServices/Private/DeviceProfileServicesUIManager.cpp b/Engine/Source/Editor/DeviceProfileServices/Private/DeviceProfileServicesUIManager.cpp
--- a/Engine/Source/Editor/DeviceProfileServices/Private/DeviceProfileServicesUIManager.cpp
+++ b/Engine/Source/Editor/DeviceProfileServices/Private/DeviceProfileServicesUIManager.cpp
@@ -1,13 +1,10 @@
 // Copyright 1998-2014 Epic Games, Inc. All Rights Reserved.
 
-/*=============================================================================
-	FDeviceProfileServicesUIManager.h: Implements the FDeviceProfileServicesUIManager class.
-=============================================================================*/
-
 #include "DeviceProfileServicesPCH.h"
 #include "DeviceProfiles/DeviceProfileManager.h"
 #include "DeviceProfiles/DeviceProfile.h"
 
+
 DEFINE_LOG_CATEGORY_STATIC(LogDeviceProfileServices, Log, All);
 
 
diff --git a/Engine/Source/Editor/DeviceProfileServices/Private/DeviceProfileServicesUIManager.h b/Engine/Source/Editor/DeviceProfileServices/Private/DeviceProfileServicesUIManager.h
--- a/Engine/Source/Editor/DeviceProfileServices/Private/DeviceProfileServicesUIManager.h
+++ b/Engine/Source/Editor/DeviceProfileServices/Private/DeviceProfileServicesUIManager.h
@@ -1,9 +1,5 @@
 // Copyright 1998-2014 Epic Games, Inc. All Rights Reserved.
 
-/*=============================================================================
-	FDeviceProfileServicesUIManager.h: Declares the FDeviceProfileServicesUIManager class.
-=============================================================================*/
-
 #pragma once
 
 
@@ -15,16 +11,15 @@ class FDeviceProfileServicesUIManager
 {
 public:
 
-	/**
-	 * Default constructor
-	 */
+	/** Default constructor. */
 	FDeviceProfileServicesUIManager( );
 
-	virtual ~FDeviceProfileServicesUIManager() {}
+	/** Destructor. */
+	virtual ~FDeviceProfileServicesUIManager() { }
 
 public:
 
-	// Begin IDeviceProfileServicesUIManager Interface
+	// IDeviceProfileServicesUIManager Interface
 
 	virtual const FName GetDeviceIconName( const FString& DeviceName ) const override;
 	virtual const TArray<TSharedPtr<FString> > GetPlatformList( ) override;
@@ -32,31 +27,25 @@ public:
 	virtual const FName GetPlatformIconName( const FString& DeviceName ) const override;
 	virtual void SetProfile( const FString& DeviceProfileName ) override;
 
-	// End IDeviceProfileServicesUIManager Interface
-
 protected:
 
-	/**
-	 * Generates the UI icon list.
-	 */
+	/** Generates the UI icon list. */
 	void CreatePlatformMap();
 
-	/**
-	 * Refresh the UI list - rebuild lists.
-	 */
+	/** Refresh the UI list - rebuild lists. */
 	void HandleRefreshUIData();
 
 private:
 
-	// Map of profiles to platform types
+	/** Map of profiles to platform types. */
 	TMap<const UClass*, FString> PickerTypeMap;
 
-	// Map of profiles to platform types
+	/** Map of profiles to platform types. */
 	TMap<const FString, FString> DeviceToPlatformMap;
 
-	// Map of profiles to platform types
+	/** Map of profiles to platform types. */
 	TMap<const FString, FName> DeviceTypeToIconMap;
 
-	// Holds the list of known platforms.
+	/** Holds the list of known platforms. */
 	TArray<TSharedPtr<FString> > PlatformList;
 };
diff --git a/Engine/Source/Editor/DeviceProfileServices/Public/DeviceProfileServices.h b/Engine/Source/Editor/DeviceProfileServices/Public/DeviceProfileServices.h
--- a/Engine/Source/Editor/DeviceProfileServices/Public/DeviceProfileServices.h
+++ b/Engine/Source/Editor/DeviceProfileServices/Public/DeviceProfileServices.h
@@ -1,11 +1,8 @@
 // Copyright 1998-2014 Epic Games, Inc. All Rights Reserved.
 
-/*=============================================================================
-	DeviceProfileServices.h: DeviceProfileServices module public header file.
-=============================================================================*/
-
 #pragma once
 
+
 /* Dependencies
  *****************************************************************************/
 
diff --git a/Engine/Source/Editor/DeviceProfileServices/Public/IDeviceProfileServicesModule.h b/Engine/Source/Editor/DeviceProfileServices/Public/IDeviceProfileServicesModule.h
--- a/Engine/Source/Editor/DeviceProfileServices/Public/IDeviceProfileServicesModule.h
+++ b/Engine/Source/Editor/DeviceProfileServices/Public/IDeviceProfileServicesModule.h
@@ -1,11 +1,8 @@
 // Copyright 1998-2014 Epic Games, Inc. All Rights Reserved.
 
-/*=============================================================================
-	IDeviceProfileServicesModule.h: Declares the IDeviceProfileServicesModule interface.
-=============================================================================*/
-
 #pragma once
 
+
 /**
  * Device Profile Editor module
  */
@@ -17,12 +14,12 @@ public:
 	/**
 	 * Gets the profile services manager.
 	 *
-	 * @return - the profile services manager.
+	 * @return The profile services manager.
 	 */
 	virtual IDeviceProfileServicesUIManagerRef GetProfileServicesManager( ) = 0;
 
-	/**
-	 * Virtual destructor.
-	 */
+public:
+
+	/** Virtual destructor. */
 	virtual ~IDeviceProfileServicesModule( ) { }
 };
diff --git a/Engine/Source/Editor/DeviceProfileServices/Public/IDeviceProfileServicesUIManager.h b/Engine/Source/Editor/DeviceProfileServices/Public/IDeviceProfileServicesUIManager.h
--- a/Engine/Source/Editor/DeviceProfileServices/Public/IDeviceProfileServicesUIManager.h
+++ b/Engine/Source/Editor/DeviceProfileServices/Public/IDeviceProfileServicesUIManager.h
@@ -1,26 +1,19 @@
 // Copyright 1998-2014 Epic Games, Inc. All Rights Reserved.
 
-/*=============================================================================
-	IDeviceProfileServicesUIManager.h: Declares the IDeviceProfileServicesUIManager interface.
-=============================================================================*/
-
 #pragma once
 
+
 class UDeviceProfile;
 
-/**
- * Type definition for shared pointers to instances of IDeviceProfileServicesUIManager.
- */
+/** Type definition for shared pointers to instances of IDeviceProfileServicesUIManager. */
 typedef TSharedPtr<class IDeviceProfileServicesUIManager> IDeviceProfileServicesUIManagerPtr;
 
-/**
- * Type definition for shared references to instances of IDeviceProfileServicesUIManager.
- */
+/** Type definition for shared references to instances of IDeviceProfileServicesUIManager. */
 typedef TSharedRef<class IDeviceProfileServicesUIManager> IDeviceProfileServicesUIManagerRef;
 
 
 /**
- * Device Profile Services Manager for UI
+ * Device Profile Services Manager for UI.
  */
 class IDeviceProfileServicesUIManager
 {
@@ -29,7 +22,7 @@ public:
 	/**
 	 * Gets an icon name for a given device.
 	 *
-	 * @param DeviceName - the device name.
+	 * @param DeviceName The device name.
 	 * @return The icon name.
 	 */
 	virtual const FName GetDeviceIconName( const FString& DeviceName ) const = 0;
@@ -44,15 +37,15 @@ public:
 	/**
 	 * Gets an array of profile that match a type.
 	 *
-	 * @param OutDeviceProfiles - the populated profile list.
-	 * @param InType - the type to match.
+	 * @param OutDeviceProfiles The populated profile list.
+	 * @param InType The type to match.
 	 */
 	virtual void GetProfilesByType( TArray<UDeviceProfile*>& OutDeviceProfiles, const FString& InType ) = 0;
 
 	/**
 	 * Gets an icon name for a given platform.
 	 *
-	 * @param DeviceName - the platform name.
+	 * @param DeviceName -The platform name.
 	 * @return The icon name.
 	 */
 	virtual const FName GetPlatformIconName( const FString& DeviceName ) const = 0;
@@ -60,7 +53,7 @@ public:
 	/**
 	 * Set a profile - add the type to an INI file so it can be used in a history view.
 	 *
-	 * @param DeviceProfileName - The platform type.
+	 * @param DeviceProfileName The platform type.
 	 */
 	virtual void SetProfile( const FString& DeviceProfileName ) = 0;
 };
diff --git a/Engine/Source/Editor/EditorStyle/Private/SlateEditorStyle.cpp b/Engine/Source/Editor/EditorStyle/Private/SlateEditorStyle.cpp
--- a/Engine/Source/Editor/EditorStyle/Private/SlateEditorStyle.cpp
+++ b/Engine/Source/Editor/EditorStyle/Private/SlateEditorStyle.cpp
@@ -1,9 +1,5 @@
 // Copyright 1998-2014 Epic Games, Inc. All Rights Reserved.
 
-/*=============================================================================
-	EditorStyle.cpp: Implements the FSlateEditorStyle class.
-=============================================================================*/
-
 #include "EditorStylePrivatePCH.h"
 
 #if (WITH_EDITOR || (IS_PROGRAM && PLATFORM_DESKTOP))
diff --git a/Engine/Source/Editor/FoliageEdit/Private/FoliageEdMode.cpp b/Engine/Source/Editor/FoliageEdit/Private/FoliageEdMode.cpp
--- a/Engine/Source/Editor/FoliageEdit/Private/FoliageEdMode.cpp
+++ b/Engine/Source/Editor/FoliageEdit/Private/FoliageEdMode.cpp
@@ -1,9 +1,5 @@
 // Copyright 1998-2014 Epic Games, Inc. All Rights Reserved.
 
-/*================================================================================
-	FoliageEdMode.cpp: Foliage editing mode
-================================================================================*/
-
 #include "UnrealEd.h"
 #include "StaticMeshResources.h"
 #include "ObjectTools.h"
diff --git a/Engine/Source/Editor/FoliageEdit/Private/FoliageEdMode.h b/Engine/Source/Editor/FoliageEdit/Private/FoliageEdMode.h
--- a/Engine/Source/Editor/FoliageEdit/Private/FoliageEdMode.h
+++ b/Engine/Source/Editor/FoliageEdit/Private/FoliageEdMode.h
@@ -1,11 +1,8 @@
 // Copyright 1998-2014 Epic Games, Inc. All Rights Reserved.
 
-/*================================================================================
-	FoliageEdMode.h: Foliage editing
-================================================================================*/
-
 #pragma once
 
+
 // Forward declarations
 class ULandscapeComponent;
 class UFoliageType;
diff --git a/Engine/Source/Editor/GeometryMode/Public/EditorGeometry.h b/Engine/Source/Editor/GeometryMode/Public/EditorGeometry.h
--- a/Engine/Source/Editor/GeometryMode/Public/EditorGeometry.h
+++ b/Engine/Source/Editor/GeometryMode/Public/EditorGeometry.h
@@ -1,15 +1,12 @@
 // Copyright 1998-2014 Epic Games, Inc. All Rights Reserved.
 
-/*=============================================================================
-	EditorGeometry.h: Support for the geometry editing mode in UnrealEd.
-
-=============================================================================*/
-
 #pragma once
 
+
 class FGeomObject;
 class FPoly;
 
+
 /**
  * Base class for all classes that support storing and editing of geometry.
  */
diff --git a/Engine/Source/Editor/GraphEditor/Private/GraphEditorSettings.cpp b/Engine/Source/Editor/GraphEditor/Private/GraphEditorSettings.cpp
--- a/Engine/Source/Editor/GraphEditor/Private/GraphEditorSettings.cpp
+++ b/Engine/Source/Editor/GraphEditor/Private/GraphEditorSettings.cpp
@@ -1,9 +1,5 @@
 // Copyright 1998-2014 Epic Games, Inc. All Rights Reserved.
 
-/*=============================================================================
-	GraphEditorSettings.cpp: Implements the UGraphEditorSettings class.
-=============================================================================*/
-
 #include "GraphEditorCommon.h"
 
 
diff --git a/Engine/Source/Editor/GraphEditor/Public/GraphEditorSettings.h b/Engine/Source/Editor/GraphEditor/Public/GraphEditorSettings.h
--- a/Engine/Source/Editor/GraphEditor/Public/GraphEditorSettings.h
+++ b/Engine/Source/Editor/GraphEditor/Public/GraphEditorSettings.h
@@ -1,10 +1,7 @@
 // Copyright 1998-2014 Epic Games, Inc. All Rights Reserved.
 
-/*=============================================================================
-	GraphEditorSettings.h: Declares the UGraphEditorSettings class.
-=============================================================================*/
-
 #pragma once
+
 #include "Runtime/SlateCore/Public/Layout/Margin.h"
 #include "GraphEditorSettings.generated.h"
 
@@ -22,8 +19,8 @@ public:
 
 #if WITH_EDITOR
 	// UObject interface
+
 	virtual void PostEditChangeProperty(FPropertyChangedEvent& PropertyChangedEvent) override;
-	// End of UObject interface
 #endif
 
 public:
@@ -65,6 +62,7 @@ public:
 	FVector2D GetShadowDeltaSize() const;
 
 public:
+
 	/** The maximum value to clamp the absolute value of the horizontal distance between endpoints when calculating tangents (when the wire is moving forward) */
 	UPROPERTY(config, EditAnywhere, Category=Splines, AdvancedDisplay)
 	float ForwardSplineHorizontalDeltaRange;
@@ -99,6 +97,7 @@ public:
 
 	/** Compute a tangent given the start and end points of the spline */
 	FVector2D ComputeSplineTangent(const FVector2D& Start, const FVector2D& End) const;
+
 public:
 
 	/** The default color is used only for types not specifically defined below.  Generally if it's seen, it means another type needs to be defined so that the wire in question can have an appropriate color. */
diff --git a/Engine/Source/Editor/InputBindingEditor/Private/InputBindingEditorModule.cpp b/Engine/Source/Editor/InputBindingEditor/Private/InputBindingEditorModule.cpp
--- a/Engine/Source/Editor/InputBindingEditor/Private/InputBindingEditorModule.cpp
+++ b/Engine/Source/Editor/InputBindingEditor/Private/InputBindingEditorModule.cpp
@@ -1,10 +1,5 @@
-
 // Copyright 1998-2014 Epic Games, Inc. All Rights Reserved.
 
-/*=============================================================================
-	InputBindingEditorModule.cpp: Implements the FInputBindingEditorModule class.
-=============================================================================*/
-
 #include "InputBindingEditorPrivatePCH.h"
 
 
@@ -16,7 +11,7 @@ class FInputBindingEditorModule
 {
 public:
 
-	// Begin IInputBindingEditorModule interface
+	// IInputBindingEditorModule interface
 
 	virtual TWeakPtr<SWidget> CreateInputBindingEditorPanel( ) override
 	{
@@ -31,12 +26,9 @@ public:
 		BindingEditorPanels.Remove(Panel.Pin());
 	}
 
-	// End IInputBindingEditorModule interface
-
-
 private:
 
-	// Holds the collection of created binding editor panels.
+	/** Holds the collection of created binding editor panels. */
 	TArray<TSharedPtr<SWidget> > BindingEditorPanels;
 };
 
diff --git a/Engine/Source/Editor/InputBindingEditor/Private/InputBindingEditorPrivatePCH.h b/Engine/Source/Editor/InputBindingEditor/Private/InputBindingEditorPrivatePCH.h
--- a/Engine/Source/Editor/InputBindingEditor/Private/InputBindingEditorPrivatePCH.h
+++ b/Engine/Source/Editor/InputBindingEditor/Private/InputBindingEditorPrivatePCH.h
@@ -1,13 +1,8 @@
 // Copyright 1998-2014 Epic Games, Inc. All Rights Reserved.
 
-/*=============================================================================
-	InputBindingEditorPrivatePCH.h: Pre-compiled header file for the InputBindingEditor module.
-=============================================================================*/
-
 #pragma once
 
-
-#include "../Public/InputBindingEditor.h"
+#include "InputBindingEditor.h"
 
 
 /* Dependencies
@@ -22,4 +17,4 @@
 #include "SGestureEditBox.h"
 #include "SGestureTreeItem.h"
 #include "SGestureEditor.h"
-#include "SInputBindingEditorPanel.h"
\ No newline at end of file
+#include "SInputBindingEditorPanel.h"
diff --git a/Engine/Source/Editor/InputBindingEditor/Private/Widgets/SGestureEditBox.cpp b/Engine/Source/Editor/InputBindingEditor/Private/Widgets/SGestureEditBox.cpp
--- a/Engine/Source/Editor/InputBindingEditor/Private/Widgets/SGestureEditBox.cpp
+++ b/Engine/Source/Editor/InputBindingEditor/Private/Widgets/SGestureEditBox.cpp
@@ -1,9 +1,5 @@
 // Copyright 1998-2014 Epic Games, Inc. All Rights Reserved.
 
-/*=============================================================================
-	SGestureEditBox.cpp: Implements the SGestureEditBox class.
-=============================================================================*/
-
 #include "InputBindingEditorPrivatePCH.h"
 
 
diff --git a/Engine/Source/Editor/InputBindingEditor/Private/Widgets/SGestureEditBox.h b/Engine/Source/Editor/InputBindingEditor/Private/Widgets/SGestureEditBox.h
--- a/Engine/Source/Editor/InputBindingEditor/Private/Widgets/SGestureEditBox.h
+++ b/Engine/Source/Editor/InputBindingEditor/Private/Widgets/SGestureEditBox.h
@@ -1,9 +1,5 @@
 // Copyright 1998-2014 Epic Games, Inc. All Rights Reserved.
 
-/*=============================================================================
-	SGestureEditBox.h: Declares the SGestureEditBox class.
-=============================================================================*/
-
 #pragma once
 
 
@@ -18,17 +14,15 @@ public:
 	SLATE_BEGIN_ARGS( SGestureEditBox ){}
 	SLATE_END_ARGS()
 
-
 public:
 
 	/**
 	 * Constructs the widget.
 	 *
-	 * @param InArgs - The Slate argument list.
+	 * @param InArgs The Slate argument list.
 	 */
 	void Construct( const FArguments& InArgs, TSharedPtr<struct FGestureTreeItem> InputCommand );
 
-
 private:
 
 	/** @return Border image for the text box based on the hovered and focused state */
@@ -64,19 +58,23 @@ private:
 	EVisibility GetNotificationVisibility() const; 
 	virtual FReply OnMouseButtonDown( const FGeometry& MyGeometry, const FPointerEvent& InMouseEvent );
 
-
 private:
 
 	/** The gesture editor for this box */
 	TSharedPtr<class SGestureEditor> GestureEditor;
+	
 	/** Menu anchor where the conflict pop-up is shown */
 	TSharedPtr<SMenuAnchor> ConflictPopup;
+	
 	/** The button for committing gesture */
 	mutable TSharedPtr<SButton> GestureAcceptButton;
+	
 	/** Styling: border image to draw when not hovered or focused */
 	const FSlateBrush* BorderImageNormal;
+	
 	/** Styling: border image to draw when hovered */
 	const FSlateBrush* BorderImageHovered;
+	
 	/** Styling: border image to draw when focused */
 	const FSlateBrush* BorderImageFocused;
 };
diff --git a/Engine/Source/Editor/InputBindingEditor/Private/Widgets/SGestureEditor.cpp b/Engine/Source/Editor/InputBindingEditor/Private/Widgets/SGestureEditor.cpp
--- a/Engine/Source/Editor/InputBindingEditor/Private/Widgets/SGestureEditor.cpp
+++ b/Engine/Source/Editor/InputBindingEditor/Private/Widgets/SGestureEditor.cpp
@@ -1,9 +1,5 @@
 // Copyright 1998-2014 Epic Games, Inc. All Rights Reserved.
 
-/*=============================================================================
-	SGestureEditor.cpp: Implements the SGestureEditor class.
-=============================================================================*/
-
 #include "InputBindingEditorPrivatePCH.h"
 
 
@@ -39,9 +35,9 @@ void SGestureEditor::Construct( const FArguments& InArgs, TSharedPtr<FGestureTre
 }
 
 
-
 TWeakPtr<SGestureEditor> SGestureEditor::GestureBeingEdited;
 
+
 FReply SGestureEditor::OnKeyDown( const FGeometry& MyGeometry, const FKeyboardEvent& InKeyboardEvent )
 {
 	const FKey Key = InKeyboardEvent.GetKey();
@@ -77,16 +73,19 @@ FReply SGestureEditor::OnKeyDown( const FGeometry& MyGeometry, const FKeyboardEv
 	return FReply::Unhandled();
 }
 
+
 FReply SGestureEditor::OnKeyUp( const FGeometry& MyGeometry, const FKeyboardEvent& InKeyboardEvent )
 {
 	return FReply::Unhandled();
 }
 
+
 FReply SGestureEditor::OnKeyChar( const FGeometry& MyGeometry, const FCharacterEvent& InCharacterEvent )
 {
 	return FReply::Unhandled();
 }
 
+
 FReply SGestureEditor::OnMouseButtonDown( const FGeometry& MyGeometry, const FPointerEvent& InMouseEvent )
 {
 	if ( InMouseEvent.GetEffectingButton() == EKeys::LeftMouseButton && !bIsEditing )
@@ -98,17 +97,20 @@ FReply SGestureEditor::OnMouseButtonDown( const FGeometry& MyGeometry, const FPo
 	return FReply::Unhandled();
 }
 
+
 FReply SGestureEditor::OnMouseButtonDoubleClick( const FGeometry& InMyGeometry, const FPointerEvent& InMouseEvent )
 {
 	return FReply::Handled();
 }
 
+
 void SGestureEditor::OnKeyboardFocusLost( const FKeyboardFocusEvent& InKeyboardFocusEvent )
 {
 	// Notify a listener that we lost focus so they can determine if we should still be in edit mode
 	OnEditBoxLostFocus.ExecuteIfBound();
 }
 
+
 FText SGestureEditor::OnGetGestureInputText() const
 {
 	if( bIsEditing )
@@ -125,6 +127,7 @@ FText SGestureEditor::OnGetGestureInputText() const
 	}
 }
 
+
 FText SGestureEditor::OnGetGestureInputHintText() const
 {
 	const TSharedPtr<const FInputGesture>& ActiveGesture = CommandInfo->GetActiveGesture();
@@ -144,6 +147,7 @@ FText SGestureEditor::OnGetGestureInputHintText() const
 	return FText::GetEmpty();
 }
 
+
 void SGestureEditor::CommitNewGesture()
 {
 	if( EditingInputGesture.IsValidGesture() )
@@ -152,11 +156,13 @@ void SGestureEditor::CommitNewGesture()
 	}
 }
 
+
 void SGestureEditor::RemoveActiveGesture()
 {
 	CommandInfo->RemoveActiveGesture();
 }
 
+
 void SGestureEditor::StartEditing()
 {
 	if( GestureBeingEdited.IsValid()  )
@@ -173,6 +179,7 @@ void SGestureEditor::StartEditing()
 	OnEditingStarted.ExecuteIfBound();
 }
 
+
 void SGestureEditor::StopEditing()
 {
 	if( GestureBeingEdited.IsValid() && GestureBeingEdited.Pin().Get() == this )
@@ -188,6 +195,7 @@ void SGestureEditor::StopEditing()
 	NotificationMessage = FText::GetEmpty();
 }
 
+
 void SGestureEditor::OnGestureTyped( const FInputGesture& NewGesture )
 {
 	// Check to see if a valid sequence was entered
@@ -214,6 +222,7 @@ void SGestureEditor::OnGestureTyped( const FInputGesture& NewGesture )
 	OnGestureChanged.ExecuteIfBound();
 }
 
+
 void SGestureEditor::OnGestureCommitted( const FInputGesture& NewGesture )
 {
 	// This delegate is only called on valid gestures
