UI is now robust enough to handle variant list corruption.
Fix for UFE Crash caused by Serialization issues (Corruption no longer occurs).

diff --git a/Engine/Source/Developer/ProjectLauncher/Private/Widgets/Deploy/SProjectLauncherDeployTargetListRow.h b/Engine/Source/Developer/ProjectLauncher/Private/Widgets/Deploy/SProjectLauncherDeployTargetListRow.h
--- a/Engine/Source/Developer/ProjectLauncher/Private/Widgets/Deploy/SProjectLauncherDeployTargetListRow.h
+++ b/Engine/Source/Developer/ProjectLauncher/Private/Widgets/Deploy/SProjectLauncherDeployTargetListRow.h
@@ -178,7 +178,7 @@ private:
 	{
 		ILauncherDeviceGroupPtr ActiveGroup = DeviceGroup.Get();
 
-		if (ActiveGroup.IsValid() && DeviceProxy.IsValid())
+		if (ActiveGroup.IsValid() && DeviceProxy.IsValid() && DeviceProxy->HasVariant(SelectedVariant))
 		{
 			const FString& DeviceID = DeviceProxy->GetTargetDeviceId(SelectedVariant);
 			if (NewState == ESlateCheckBoxState::Checked)
@@ -224,7 +224,7 @@ private:
 	// Callback for changing the selected variant
 	void HandleVariantSelectorVariantSelected(FName InVariant)
 	{
-		if (DeviceProxy.IsValid())
+		if (DeviceProxy.IsValid() && DeviceProxy->HasVariant(SelectedVariant) && DeviceProxy->HasVariant(InVariant))
 		{
 			if (bDeviceInGroup)
 			{
@@ -245,8 +245,15 @@ private:
 	// Callback for getting the icon image of the device.
 	const FSlateBrush* HandleDeviceImage() const
 	{
-		const PlatformInfo::FPlatformInfo* const PlatformInfo = PlatformInfo::FindPlatformInfo(*DeviceProxy->GetTargetPlatformName(NAME_None));
-		return (PlatformInfo) ? FEditorStyle::GetBrush(PlatformInfo->GetIconStyleName(PlatformInfo::EPlatformIconSize::Normal)) : FStyleDefaults::GetNoBrush();
+		if (DeviceProxy->HasVariant(NAME_None))
+		{
+			const PlatformInfo::FPlatformInfo* const PlatformInfo = PlatformInfo::FindPlatformInfo(*DeviceProxy->GetTargetPlatformName(NAME_None));
+			if (PlatformInfo)
+			{
+				FEditorStyle::GetBrush(PlatformInfo->GetIconStyleName(PlatformInfo::EPlatformIconSize::Normal));
+			}
+		}
+		return FStyleDefaults::GetNoBrush();
 	}
 
 	// Callback for getting the friendly name.
@@ -277,7 +284,11 @@ private:
 	// Callback for getting the host platform name.
 	FString HandleHostPlatformText() const
 	{
-		return DeviceProxy->GetTargetPlatformName(NAME_None);
+		if (DeviceProxy->HasVariant(NAME_None))
+		{
+			return DeviceProxy->GetTargetPlatformName(NAME_None);
+		}
+		return LOCTEXT("InvalidVariant", "Invalid Variant").ToString();
 	}
 
 	// Callback for getting the default variant name in the case where 
diff --git a/Engine/Source/Developer/ProjectLauncher/Private/Widgets/Deploy/SProjectLauncherDeployTargets.cpp b/Engine/Source/Developer/ProjectLauncher/Private/Widgets/Deploy/SProjectLauncherDeployTargets.cpp
--- a/Engine/Source/Developer/ProjectLauncher/Private/Widgets/Deploy/SProjectLauncherDeployTargets.cpp
+++ b/Engine/Source/Developer/ProjectLauncher/Private/Widgets/Deploy/SProjectLauncherDeployTargets.cpp
@@ -154,8 +154,16 @@ FText SProjectLauncherDeployTargets::HandleDeviceListRowToolTipText(ITargetDevic
 {
 	FTextBuilder Builder;
 	Builder.AppendLineFormat(LOCTEXT("DeviceListRowToolTipName", "Name: {0}"), FText::FromString(DeviceProxy->GetName()));
-	Builder.AppendLineFormat(LOCTEXT("DeviceListRowToolTipPlatform", "Platform: {0}"), FText::FromString(DeviceProxy->GetTargetPlatformName(NAME_None /*@Todo: Fix this!*/)));
-	Builder.AppendLineFormat(LOCTEXT("DeviceListRowToolTipDeviceId", "Device ID: {0}"), FText::FromString(DeviceProxy->GetTargetDeviceId(NAME_None /*@Todo: Fix this!*/)));
+	
+	if (DeviceProxy->HasVariant(NAME_None))
+	{	
+		Builder.AppendLineFormat(LOCTEXT("DeviceListRowToolTipPlatform", "Platform: {0}"), FText::FromString(DeviceProxy->GetTargetPlatformName(NAME_None)));
+		Builder.AppendLineFormat(LOCTEXT("DeviceListRowToolTipDeviceId", "Device ID: {0}"), FText::FromString(DeviceProxy->GetTargetDeviceId(NAME_None)));
+	}
+	else
+	{
+		Builder.AppendLine(LOCTEXT("InvalidDevice", "Invalid Device"));
+	}
 
 	return Builder.ToText();
 }
diff --git a/Engine/Source/Developer/ProjectLauncher/Private/Widgets/Deploy/SProjectLauncherSimpleDeviceListRow.h b/Engine/Source/Developer/ProjectLauncher/Private/Widgets/Deploy/SProjectLauncherSimpleDeviceListRow.h
--- a/Engine/Source/Developer/ProjectLauncher/Private/Widgets/Deploy/SProjectLauncherSimpleDeviceListRow.h
+++ b/Engine/Source/Developer/ProjectLauncher/Private/Widgets/Deploy/SProjectLauncherSimpleDeviceListRow.h
@@ -333,30 +333,41 @@ private:
 	{
 		check(DeviceProxy.IsValid());
 
-		FName Variant = SimpleProfile->GetDeviceVariant();
-		
-		// If the profile refers to a variant no longer supported by this device fall back to the default.
-		if (!DeviceProxy->HasVariant(Variant))
+		// This detects a corrupt device, it doesn't even have data for it's default variant Remove the profile.
+		// Entry will show invalid.
+		if (!DeviceProxy->HasVariant(NAME_None))
 		{
-			Variant = NAME_None;
-			SimpleProfile->SetDeviceVariant(Variant);
+			Model->GetProfileManager()->RemoveProfile(LaunchProfile.ToSharedRef());
+			LaunchProfile = nullptr;
 		}
+		
+		if (LaunchProfile.IsValid())
+		{
+			FName Variant = SimpleProfile->GetDeviceVariant();
+		
+			// If the profile refers to a variant no longer supported by this device fall back to the default.
+			if (!DeviceProxy->HasVariant(Variant))
+			{
+				Variant = NAME_None;
+				SimpleProfile->SetDeviceVariant(Variant);
+			}
 
-		// Setup the Profile
-		LaunchProfile->SetDeploymentMode(ELauncherProfileDeploymentModes::FileServer);
+			// Setup the Profile
+			LaunchProfile->SetDeploymentMode(ELauncherProfileDeploymentModes::FileServer);
 
-		ILauncherDeviceGroupRef NewGroup = Model->GetProfileManager()->CreateUnmanagedDeviceGroup();
-		NewGroup->AddDevice(DeviceProxy->GetTargetDeviceId(Variant));
-		LaunchProfile->SetDeployedDeviceGroup(NewGroup);
+			ILauncherDeviceGroupRef NewGroup = Model->GetProfileManager()->CreateUnmanagedDeviceGroup();
+			NewGroup->AddDevice(DeviceProxy->GetTargetDeviceId(Variant));
+			LaunchProfile->SetDeployedDeviceGroup(NewGroup);
 
-		LaunchProfile->ClearCookedCultures();
-		LaunchProfile->AddCookedPlatform(DeviceProxy->GetTargetPlatformName(Variant));
+			LaunchProfile->ClearCookedCultures();
+			LaunchProfile->AddCookedPlatform(DeviceProxy->GetTargetPlatformName(Variant));
 
-		bool Advanced = IsAdvanced.IsBound() && IsAdvanced.Get();
-		if (Advanced)
-		{
-			LaunchProfile->SetBuildConfiguration(SimpleProfile->GetBuildConfiguration());
-			LaunchProfile->SetCookMode(SimpleProfile->GetCookMode());
+			bool Advanced = IsAdvanced.IsBound() && IsAdvanced.Get();
+			if (Advanced)
+			{
+				LaunchProfile->SetBuildConfiguration(SimpleProfile->GetBuildConfiguration());
+				LaunchProfile->SetCookMode(SimpleProfile->GetCookMode());
+			}
 		}
 	}
 
@@ -385,8 +396,15 @@ private:
 	// Callback for getting the icon image of the device.
 	const FSlateBrush* HandleDeviceImage( ) const
 	{
-		const PlatformInfo::FPlatformInfo* const PlatformInfo = PlatformInfo::FindPlatformInfo(*DeviceProxy->GetTargetPlatformName(SimpleProfile->GetDeviceVariant()));
-		return (PlatformInfo) ? FEditorStyle::GetBrush(PlatformInfo->GetIconStyleName(PlatformInfo::EPlatformIconSize::Large)) : FStyleDefaults::GetNoBrush();
+		if (LaunchProfile.IsValid())
+		{
+			const PlatformInfo::FPlatformInfo* const PlatformInfo = PlatformInfo::FindPlatformInfo(*DeviceProxy->GetTargetPlatformName(SimpleProfile->GetDeviceVariant()));
+			if (PlatformInfo)
+			{
+				return FEditorStyle::GetBrush(PlatformInfo->GetIconStyleName(PlatformInfo::EPlatformIconSize::Large));
+			}
+		}
+		return FStyleDefaults::GetNoBrush();
 	}
 
 	// Callback for getting the friendly name.
@@ -417,7 +435,11 @@ private:
 	// Callback for getting the host platform name.
 	FString HandleHostPlatformText( ) const
 	{
-		return DeviceProxy->GetTargetPlatformName(SimpleProfile->GetDeviceVariant());
+		if (LaunchProfile.IsValid())
+		{
+			return DeviceProxy->GetTargetPlatformName(SimpleProfile->GetDeviceVariant());
+		}
+		return LOCTEXT("InvalidVariant", "Invalid Variant").ToString();
 	}
 
 private:
diff --git a/Engine/Source/Developer/ProjectLauncher/Private/Widgets/SProjectLauncherValidation.h b/Engine/Source/Developer/ProjectLauncher/Private/Widgets/SProjectLauncherValidation.h
--- a/Engine/Source/Developer/ProjectLauncher/Private/Widgets/SProjectLauncherValidation.h
+++ b/Engine/Source/Developer/ProjectLauncher/Private/Widgets/SProjectLauncherValidation.h
@@ -174,14 +174,10 @@ private:
 	EVisibility HandleValidationMessageVisibility( ELauncherProfileValidationErrors::Type Error ) const
 	{
 		ILauncherProfilePtr LaunchProfile = LaunchProfileAttr.Get();
-		if (LaunchProfile.IsValid())
+		if (!LaunchProfile.IsValid() || LaunchProfile->HasValidationError(Error))
 		{
-			if (LaunchProfile->HasValidationError(Error))
-			{
-				return EVisibility::Visible;
-			}
+			return EVisibility::Visible;
 		}
-
 		return EVisibility::Collapsed;
 	}
 
@@ -194,9 +190,10 @@ private:
 			{
 				return LOCTEXT("NoPlatformSDKInstalled", "A required platform SDK is mising: ").ToString() + LaunchProfile->GetInvalidPlatform();
 			}
-		}
 
-		return TEXT("");
+			return TEXT("");
+		}
+		return LOCTEXT("InvalidLaunchProfile", "Invalid Launch Profile.").ToString();
 	}
 
 private:
diff --git a/Engine/Source/Developer/TargetDeviceServices/Classes/TargetDeviceServiceMessages.h b/Engine/Source/Developer/TargetDeviceServices/Classes/TargetDeviceServiceMessages.h
--- a/Engine/Source/Developer/TargetDeviceServices/Classes/TargetDeviceServiceMessages.h
+++ b/Engine/Source/Developer/TargetDeviceServices/Classes/TargetDeviceServiceMessages.h
@@ -531,7 +531,16 @@ struct FTargetDeviceVariant
 	FName VanillaPlatformId;
 
 	UPROPERTY()
-	FText PlatformDisplayName;
+	FString PlatformDisplayName;
+};
+
+template<>
+struct TStructOpsTypeTraits<FTargetDeviceVariant> : public TStructOpsTypeTraitsBase
+{
+	enum
+	{
+		WithMessageHandling = true
+	};
 };
 
 
diff --git a/Engine/Source/Developer/TargetDeviceServices/Private/TargetDeviceProxy.cpp b/Engine/Source/Developer/TargetDeviceServices/Private/TargetDeviceProxy.cpp
--- a/Engine/Source/Developer/TargetDeviceServices/Private/TargetDeviceProxy.cpp
+++ b/Engine/Source/Developer/TargetDeviceServices/Private/TargetDeviceProxy.cpp
@@ -69,7 +69,7 @@ void FTargetDeviceProxy::UpdateFromMessage( const FTargetDeviceServicePong& Mess
 			Variant.TargetPlatformName = MsgVariant.TargetPlatformName;
 			Variant.TargetPlatformId = MsgVariant.TargetPlatformId;
 			Variant.VanillaPlatformId = MsgVariant.VanillaPlatformId;
-			Variant.PlatformDisplayName = MsgVariant.PlatformDisplayName;
+			Variant.PlatformDisplayName = FText::FromString(MsgVariant.PlatformDisplayName);
 		}
 
 		LastUpdateTime = FDateTime::UtcNow();
@@ -92,6 +92,10 @@ int32 FTargetDeviceProxy::GetVariants(TArray<FName>& OutVariants) const
 
 bool FTargetDeviceProxy::HasVariant(FName InVariant) const
 {
+	if (InVariant == NAME_None)
+	{
+		return TargetDeviceVariants.Contains(DefaultVariant);
+	}
 	return TargetDeviceVariants.Contains(InVariant);
 }
 
diff --git a/Engine/Source/Developer/TargetDeviceServices/Private/TargetDeviceService.cpp b/Engine/Source/Developer/TargetDeviceServices/Private/TargetDeviceService.cpp
--- a/Engine/Source/Developer/TargetDeviceServices/Private/TargetDeviceService.cpp
+++ b/Engine/Source/Developer/TargetDeviceServices/Private/TargetDeviceService.cpp
@@ -123,7 +123,7 @@ void FTargetDeviceService::HandlePingMessage(const FTargetDeviceServicePing& InM
 				Variant.TargetPlatformName = TargetDevice->GetTargetPlatform().PlatformName();
 				Variant.TargetPlatformId = Info.TargetPlatformName;
 				Variant.VanillaPlatformId = Info.VanillaPlatformName;
-				Variant.PlatformDisplayName = Info.DisplayName;
+				Variant.PlatformDisplayName = Info.DisplayName.ToString();
 			}
 
 			MessageEndpoint->Send(Message, Context->GetSender());
diff --git a/Engine/Source/Editor/UnrealEd/Classes/Settings/EditorExperimentalSettings.h b/Engine/Source/Editor/UnrealEd/Classes/Settings/EditorExperimentalSettings.h
--- a/Engine/Source/Editor/UnrealEd/Classes/Settings/EditorExperimentalSettings.h
+++ b/Engine/Source/Editor/UnrealEd/Classes/Settings/EditorExperimentalSettings.h
@@ -30,7 +30,7 @@ public:
 	bool bEnableEditorUtilityBlueprints;
 
 	/** The Game Launcher provides advanced workflows for packaging, deploying and launching your games. */
-	UPROPERTY(EditAnywhere, config, Category=Tools, meta=(DisplayName="Game Launcher"))
+	UPROPERTY(EditAnywhere, config, Category=Tools, meta=(DisplayName="Project Launcher"))
 	bool bGameLauncher;
 
 	/** The Messaging Debugger provides a visual utility for debugging the messaging system. */
