Fix for By the book not working in simple profiles.

diff --git a/Engine/Source/Developer/ProjectLauncher/Private/Widgets/Deploy/SProjectLauncherSimpleDeviceListRow.h b/Engine/Source/Developer/ProjectLauncher/Private/Widgets/Deploy/SProjectLauncherSimpleDeviceListRow.h
--- a/Engine/Source/Developer/ProjectLauncher/Private/Widgets/Deploy/SProjectLauncherSimpleDeviceListRow.h
+++ b/Engine/Source/Developer/ProjectLauncher/Private/Widgets/Deploy/SProjectLauncherSimpleDeviceListRow.h
@@ -359,7 +359,7 @@ private:
 			NewGroup->AddDevice(DeviceProxy->GetTargetDeviceId(Variant));
 			LaunchProfile->SetDeployedDeviceGroup(NewGroup);
 
-			LaunchProfile->ClearCookedCultures();
+			LaunchProfile->ClearCookedPlatforms();
 			LaunchProfile->AddCookedPlatform(DeviceProxy->GetTargetPlatformName(Variant));
 
 			bool Advanced = IsAdvanced.IsBound() && IsAdvanced.Get();
