Editor - Fixing log statement text coloring in the project launcher, also adding some space between the log and the bottom row of buttons.

diff --git a/Engine/Source/Developer/ProjectLauncher/Private/Widgets/SProjectLauncherProgress.h b/Engine/Source/Developer/ProjectLauncher/Private/Widgets/SProjectLauncherProgress.h
--- a/Engine/Source/Developer/ProjectLauncher/Private/Widgets/SProjectLauncherProgress.h
+++ b/Engine/Source/Developer/ProjectLauncher/Private/Widgets/SProjectLauncherProgress.h
@@ -129,7 +129,7 @@ public:
 
 			+ SVerticalBox::Slot()
 			.AutoHeight()
-			.Padding(0.0f, 0.0f, 0.0f, 0.0f)
+			.Padding(0.0f, 5.0f, 0.0f, 0.0f)
 			[
 				SNew(SHorizontalBox)
 
@@ -253,11 +253,11 @@ private:
 		{
 			Verbosity = ELogVerbosity::Display;
 		}
-		else if (InMessage.Contains(TEXT("ERROR:"), ESearchCase::CaseSensitive))
+		else if ( InMessage.Contains(TEXT("Error:"), ESearchCase::IgnoreCase) )
 		{
 			Verbosity = ELogVerbosity::Error;
 		}
-		else if (InMessage.Contains(TEXT("WARNING:"), ESearchCase::CaseSensitive))
+		else if ( InMessage.Contains(TEXT("Warning:"), ESearchCase::IgnoreCase) )
 		{
 			Verbosity = ELogVerbosity::Warning;
 		}
