Rocket version of UFE will no longer show Debug and Test.

diff --git a/Engine/Source/Developer/ProjectLauncher/Private/Widgets/Shared/SProjectLauncherBuildConfigurationSelector.h b/Engine/Source/Developer/ProjectLauncher/Private/Widgets/Shared/SProjectLauncherBuildConfigurationSelector.h
--- a/Engine/Source/Developer/ProjectLauncher/Private/Widgets/Shared/SProjectLauncherBuildConfigurationSelector.h
+++ b/Engine/Source/Developer/ProjectLauncher/Private/Widgets/Shared/SProjectLauncherBuildConfigurationSelector.h
@@ -43,8 +43,11 @@ public:
 		// create build configurations menu
 		FMenuBuilder MenuBuilder(true, NULL);
 		{
-			FUIAction DebugAction(FExecuteAction::CreateSP(this, &SProjectLauncherBuildConfigurationSelector::HandleMenuEntryClicked, EBuildConfigurations::Debug));
-			MenuBuilder.AddMenuEntry(LOCTEXT("DebugAction", "Debug"), LOCTEXT("DebugActionHint", "Debug configuration."), FSlateIcon(), DebugAction);
+			if (!FRocketSupport::IsRocket())
+			{
+				FUIAction DebugAction(FExecuteAction::CreateSP(this, &SProjectLauncherBuildConfigurationSelector::HandleMenuEntryClicked, EBuildConfigurations::Debug));
+				MenuBuilder.AddMenuEntry(LOCTEXT("DebugAction", "Debug"), LOCTEXT("DebugActionHint", "Debug configuration."), FSlateIcon(), DebugAction);
+			}
 
 			FUIAction DebugGameAction(FExecuteAction::CreateSP(this, &SProjectLauncherBuildConfigurationSelector::HandleMenuEntryClicked, EBuildConfigurations::DebugGame));
 			MenuBuilder.AddMenuEntry(LOCTEXT("DebugGameAction", "DebugGame"), LOCTEXT("DebugGameActionHint", "DebugGame configuration."), FSlateIcon(), DebugGameAction);
@@ -55,8 +58,11 @@ public:
 			FUIAction ShippingAction(FExecuteAction::CreateSP(this, &SProjectLauncherBuildConfigurationSelector::HandleMenuEntryClicked, EBuildConfigurations::Shipping));
 			MenuBuilder.AddMenuEntry(LOCTEXT("ShippingAction", "Shipping"), LOCTEXT("ShippingActionHint", "Shipping configuration."), FSlateIcon(), ShippingAction);
 
-			FUIAction TestAction(FExecuteAction::CreateSP(this, &SProjectLauncherBuildConfigurationSelector::HandleMenuEntryClicked, EBuildConfigurations::Test));
-			MenuBuilder.AddMenuEntry(LOCTEXT("TestAction", "Test"), LOCTEXT("TestActionHint", "Test configuration."), FSlateIcon(), TestAction);
+			if (!FRocketSupport::IsRocket())
+			{
+				FUIAction TestAction(FExecuteAction::CreateSP(this, &SProjectLauncherBuildConfigurationSelector::HandleMenuEntryClicked, EBuildConfigurations::Test));
+				MenuBuilder.AddMenuEntry(LOCTEXT("TestAction", "Test"), LOCTEXT("TestActionHint", "Test configuration."), FSlateIcon(), TestAction);
+			}
 		}
 
 		FSlateFontInfo TextFont = InArgs._Font.IsSet() ? InArgs._Font.Get() : FCoreStyle::Get().GetFontStyle(TEXT("SmallFont"));
