Moved pak file option to the cook by the book settings in the project launcher
UE-3402
Added check box to the advanced settings in build to build UAT as part of the launch
UE-3732
#ufe

diff --git a/Engine/Source/Developer/LauncherServices/Private/Launcher/LauncherUATTask.h b/Engine/Source/Developer/LauncherServices/Private/Launcher/LauncherUATTask.h
--- a/Engine/Source/Developer/LauncherServices/Private/Launcher/LauncherUATTask.h
+++ b/Engine/Source/Developer/LauncherServices/Private/Launcher/LauncherUATTask.h
@@ -41,7 +41,7 @@ protected:
 	{
 		if (FirstTimeCompile)
 		{
-			NoCompile = !FParse::Param( FCommandLine::Get(), TEXT("development") ) ? TEXT(" -nocompile") : TEXT("");
+			NoCompile = (!FParse::Param( FCommandLine::Get(), TEXT("development") ) && !ChainState.Profile->IsBuildingUAT()) ? TEXT(" -nocompile") : TEXT("");
 			FirstTimeCompile = false;
 		}
 
diff --git a/Engine/Source/Developer/LauncherServices/Private/Profiles/LauncherProfile.h b/Engine/Source/Developer/LauncherServices/Private/Profiles/LauncherProfile.h
--- a/Engine/Source/Developer/LauncherServices/Private/Profiles/LauncherProfile.h
+++ b/Engine/Source/Developer/LauncherServices/Private/Profiles/LauncherProfile.h
@@ -398,6 +398,11 @@ public:
 		return BuildGame;
 	}
 
+	virtual bool IsBuildingUAT() const override
+	{
+		return BuildUAT;
+	}
+
 	virtual bool IsCookingIncrementally( ) const override
 	{
 		return CookIncremental;
@@ -575,6 +580,7 @@ public:
 
 		// default build settings
 		BuildGame = false;
+		BuildUAT = false;
 
 		// default cook settings
 		CookConfiguration = FApp::GetBuildConfiguration();
@@ -631,6 +637,16 @@ public:
 		}
 	}
 
+	virtual void SetBuildUAT(bool Build) override
+	{
+		if (BuildUAT != Build)
+		{
+			BuildUAT = Build;
+
+			Validate();
+		}
+	}
+
 	virtual void SetBuildConfiguration( EBuildConfigurations::Type Configuration ) override
 	{
 		if (BuildConfiguration != Configuration)
@@ -1074,6 +1090,9 @@ private:
 	// Holds a flag indicating whether the game should be built
 	bool BuildGame;
 
+	// Holds a flag indicating whether UAT should be built
+	bool BuildUAT;
+
 	// Holds a flag indicating whether only modified content should be cooked.
 	bool CookIncremental;
 
diff --git a/Engine/Source/Developer/LauncherServices/Public/Interfaces/ILauncherProfile.h b/Engine/Source/Developer/LauncherServices/Public/Interfaces/ILauncherProfile.h
--- a/Engine/Source/Developer/LauncherServices/Public/Interfaces/ILauncherProfile.h
+++ b/Engine/Source/Developer/LauncherServices/Public/Interfaces/ILauncherProfile.h
@@ -555,6 +555,14 @@ public:
 	 * @see SetBuildGame
 	 */
 	virtual bool IsBuilding() const = 0;
+    
+	/**
+	 * Checks whether UAT should be built.
+	 *
+	 * @return true if building UAT, false otherwise.
+	 * @see SetBuildGame
+	 */
+	virtual bool IsBuildingUAT() const = 0;
 
 	/**
 	 * Checks whether incremental cooking is enabled.
@@ -715,6 +723,14 @@ public:
 	 */
 	virtual void SetBuildGame( bool Build ) = 0;
 
+	/**
+	 * Sets whether to build UAT.
+	 *
+	 * @param Build Whether UAT should be built.
+	 * @see IsBuilding
+	 */
+	virtual void SetBuildUAT( bool Build ) = 0;
+
 	/**
 	 * Sets the build configuration.
 	 *
diff --git a/Engine/Source/Developer/ProjectLauncher/Private/Widgets/Build/SProjectLauncherBuildPage.cpp b/Engine/Source/Developer/ProjectLauncher/Private/Widgets/Build/SProjectLauncherBuildPage.cpp
--- a/Engine/Source/Developer/ProjectLauncher/Private/Widgets/Build/SProjectLauncherBuildPage.cpp
+++ b/Engine/Source/Developer/ProjectLauncher/Private/Widgets/Build/SProjectLauncherBuildPage.cpp
@@ -105,7 +105,22 @@ void SProjectLauncherBuildPage::Construct( const FArguments& InArgs, const FProj
                         .IsEnabled( this, &SProjectLauncherBuildPage::HandleGenDSYMButtonEnabled )
                         .OnClicked( this, &SProjectLauncherBuildPage::HandleGenDSYMClicked )
                     ]
-                ]
+					+ SVerticalBox::Slot()
+					.AutoHeight()
+					[
+						// build mode check box
+						SNew(SCheckBox)
+						.IsChecked(this, &SProjectLauncherBuildPage::HandleUATIsChecked)
+						.OnCheckStateChanged(this, &SProjectLauncherBuildPage::HandleUATCheckedStateChanged)
+						.Padding(FMargin(4.0f, 0.0f))
+						.ToolTipText(LOCTEXT("UATCheckBoxTooltip", "If checked, UAT will be built as part of the build."))
+						.Content()
+						[
+							SNew(STextBlock)
+							.Text(LOCTEXT("UATCheckBoxText", "Build UAT"))
+						]
+					]
+               ]
             ]
 
 		+ SVerticalBox::Slot()
@@ -289,5 +304,32 @@ EVisibility SProjectLauncherBuildPage::HandleValidationErrorIconVisibility(ELaun
 	return EVisibility::Hidden;
 }
 
+void SProjectLauncherBuildPage::HandleUATCheckedStateChanged( ESlateCheckBoxState::Type CheckState )
+{
+	ILauncherProfilePtr SelectedProfile = Model->GetSelectedProfile();
+
+	if (SelectedProfile.IsValid())
+	{
+		SelectedProfile->SetBuildUAT(CheckState == ESlateCheckBoxState::Checked);
+	}
+}
+
+
+ESlateCheckBoxState::Type SProjectLauncherBuildPage::HandleUATIsChecked() const
+{
+	ILauncherProfilePtr SelectedProfile = Model->GetSelectedProfile();
+
+	if (SelectedProfile.IsValid())
+	{
+		if (SelectedProfile->IsBuildingUAT())
+		{
+			return ESlateCheckBoxState::Checked;
+		}
+	}
+
+	return ESlateCheckBoxState::Unchecked;
+}
+
+
 
 #undef LOCTEXT_NAMESPACE
diff --git a/Engine/Source/Developer/ProjectLauncher/Private/Widgets/Build/SProjectLauncherBuildPage.h b/Engine/Source/Developer/ProjectLauncher/Private/Widgets/Build/SProjectLauncherBuildPage.h
--- a/Engine/Source/Developer/ProjectLauncher/Private/Widgets/Build/SProjectLauncherBuildPage.h
+++ b/Engine/Source/Developer/ProjectLauncher/Private/Widgets/Build/SProjectLauncherBuildPage.h
@@ -72,6 +72,12 @@ private:
 	/** Callback for determining the visibility of a validation error icon. */
 	EVisibility HandleValidationErrorIconVisibility(ELauncherProfileValidationErrors::Type Error) const;
 
+	/** Callback for changing the checked state of a platform menu check box. */
+	void HandleUATCheckedStateChanged( ESlateCheckBoxState::Type CheckState );
+
+	/** Callback for determining whether a platform menu entry is checked. */
+	ESlateCheckBoxState::Type HandleUATIsChecked() const;
+
 private:
 
 	/** Holds a pointer to the data model. */
diff --git a/Engine/Source/Developer/ProjectLauncher/Private/Widgets/Cook/SProjectLauncherCookByTheBookSettings.cpp b/Engine/Source/Developer/ProjectLauncher/Private/Widgets/Cook/SProjectLauncherCookByTheBookSettings.cpp
--- a/Engine/Source/Developer/ProjectLauncher/Private/Widgets/Cook/SProjectLauncherCookByTheBookSettings.cpp
+++ b/Engine/Source/Developer/ProjectLauncher/Private/Widgets/Cook/SProjectLauncherCookByTheBookSettings.cpp
@@ -365,6 +365,23 @@ TSharedRef<SWidget> SProjectLauncherCookByTheBookSettings::MakeComplexWidget()
 						]
 					]
 
+				+ SVerticalBox::Slot()
+					.AutoHeight()
+					.Padding(0.0f, 4.0f, 0.0f, 0.0f)
+					[
+						// unreal pak check box
+						SNew(SCheckBox)
+						.IsChecked(this, &SProjectLauncherCookByTheBookSettings::HandleUnrealPakCheckBoxIsChecked)
+						.OnCheckStateChanged(this, &SProjectLauncherCookByTheBookSettings::HandleUnrealPakCheckBoxCheckStateChanged)
+						.Padding(FMargin(4.0f, 0.0f))
+						.ToolTipText(LOCTEXT("UnrealPakCheckBoxTooltip", "If checked, the content will be deployed as a single UnrealPak file instead of many separate files."))
+						.Content()
+						[
+							SNew(STextBlock)
+							.Text(LOCTEXT("UnrealPakCheckBoxText", "Store all content in a single file (UnrealPak)"))
+						]
+					]
+
 				+ SVerticalBox::Slot()
 					.AutoHeight()
 					.Padding(0.0f, 12.0f, 0.0f, 0.0f)
@@ -633,6 +650,23 @@ TSharedRef<SWidget> SProjectLauncherCookByTheBookSettings::MakeSimpleWidget()
 						]
 					]
 
+				+ SVerticalBox::Slot()
+					.AutoHeight()
+					.Padding(0.0f, 4.0f, 0.0f, 0.0f)
+					[
+						// unreal pak check box
+						SNew(SCheckBox)
+						.IsChecked(this, &SProjectLauncherCookByTheBookSettings::HandleUnrealPakCheckBoxIsChecked)
+						.OnCheckStateChanged(this, &SProjectLauncherCookByTheBookSettings::HandleUnrealPakCheckBoxCheckStateChanged)
+						.Padding(FMargin(4.0f, 0.0f))
+						.ToolTipText(LOCTEXT("UnrealPakCheckBoxTooltip", "If checked, the content will be deployed as a single UnrealPak file instead of many separate files."))
+						.Content()
+						[
+							SNew(STextBlock)
+							.Text(LOCTEXT("UnrealPakCheckBoxText", "Store all content in a single file (UnrealPak)"))
+						]
+					]
+
 				+ SVerticalBox::Slot()
 					.AutoHeight()
 					.Padding(0.0f, 12.0f, 0.0f, 0.0f)
@@ -1011,4 +1045,30 @@ void SProjectLauncherCookByTheBookSettings::HandleCookerOptionsCommitted(const F
 	}
 }
 
+void SProjectLauncherCookByTheBookSettings::HandleUnrealPakCheckBoxCheckStateChanged( ESlateCheckBoxState::Type NewState )
+{
+	ILauncherProfilePtr SelectedProfile = Model->GetSelectedProfile();
+
+	if (SelectedProfile.IsValid())
+	{
+		SelectedProfile->SetDeployWithUnrealPak(NewState == ESlateCheckBoxState::Checked);
+	}
+}
+
+
+ESlateCheckBoxState::Type SProjectLauncherCookByTheBookSettings::HandleUnrealPakCheckBoxIsChecked( ) const
+{
+	ILauncherProfilePtr SelectedProfile = Model->GetSelectedProfile();
+
+	if (SelectedProfile.IsValid())
+	{
+		if (SelectedProfile->IsPackingWithUnrealPak())
+		{
+			return ESlateCheckBoxState::Checked;
+		}
+	}
+
+	return ESlateCheckBoxState::Unchecked;
+}
+
 #undef LOCTEXT_NAMESPACE
diff --git a/Engine/Source/Developer/ProjectLauncher/Private/Widgets/Cook/SProjectLauncherCookByTheBookSettings.h b/Engine/Source/Developer/ProjectLauncher/Private/Widgets/Cook/SProjectLauncherCookByTheBookSettings.h
--- a/Engine/Source/Developer/ProjectLauncher/Private/Widgets/Cook/SProjectLauncherCookByTheBookSettings.h
+++ b/Engine/Source/Developer/ProjectLauncher/Private/Widgets/Cook/SProjectLauncherCookByTheBookSettings.h
@@ -115,6 +115,12 @@ private:
 	/** Callback for updating any settings after the selected project has changed in the profile. */
 	void HandleProfileProjectChanged();
 
+	// Callback for check state changes of the 'UnrealPak' check box.
+	void HandleUnrealPakCheckBoxCheckStateChanged( ESlateCheckBoxState::Type NewState );
+
+	// Callback for determining the checked state of the 'UnrealPak' check box.
+	ESlateCheckBoxState::Type HandleUnrealPakCheckBoxIsChecked( ) const;
+
 	/** creates the complex widget. */
 	TSharedRef<SWidget> MakeComplexWidget();
 
diff --git a/Engine/Source/Developer/ProjectLauncher/Private/Widgets/Deploy/SProjectLauncherDeployToDeviceSettings.cpp b/Engine/Source/Developer/ProjectLauncher/Private/Widgets/Deploy/SProjectLauncherDeployToDeviceSettings.cpp
--- a/Engine/Source/Developer/ProjectLauncher/Private/Widgets/Deploy/SProjectLauncherDeployToDeviceSettings.cpp
+++ b/Engine/Source/Developer/ProjectLauncher/Private/Widgets/Deploy/SProjectLauncherDeployToDeviceSettings.cpp
@@ -30,7 +30,7 @@ void SProjectLauncherDeployToDeviceSettings::Construct( const FArguments& InArgs
 					]
 			]
 
-		+ SVerticalBox::Slot()
+/*		+ SVerticalBox::Slot()
 			.AutoHeight()
 			.Padding(0.0f, 8.0f, 0.0f, 0.0f)
 			[
@@ -46,26 +46,9 @@ void SProjectLauncherDeployToDeviceSettings::Construct( const FArguments& InArgs
 							.Padding(8.0f)
 							.BodyContent()
 							[
-								SNew(SVerticalBox)
-
-								+ SVerticalBox::Slot()
-									.AutoHeight()
-									[
-										// unreal pak check box
-										SNew(SCheckBox)
-											.IsChecked(this, &SProjectLauncherDeployToDeviceSettings::HandleUnrealPakCheckBoxIsChecked)
-											.OnCheckStateChanged(this, &SProjectLauncherDeployToDeviceSettings::HandleUnrealPakCheckBoxCheckStateChanged)
-											.Padding(FMargin(4.0f, 0.0f))
-											.ToolTipText(LOCTEXT("UnrealPakCheckBoxTooltip", "If checked, the content will be deployed as a single UnrealPak file instead of many separate files."))
-											.Content()
-											[
-												SNew(STextBlock)
-													.Text(LOCTEXT("UnrealPakCheckBoxText", "Store all content in a single file (UnrealPak)"))
-											]
-									]
 							]
 					]
-			]
+			]*/
 	];
 }
 
@@ -73,31 +56,5 @@ void SProjectLauncherDeployToDeviceSettings::Construct( const FArguments& InArgs
 /* SProjectLauncherDeployToDeviceSettings callbacks
  *****************************************************************************/
 
-void SProjectLauncherDeployToDeviceSettings::HandleUnrealPakCheckBoxCheckStateChanged( ESlateCheckBoxState::Type NewState )
-{
-	ILauncherProfilePtr SelectedProfile = Model->GetSelectedProfile();
-
-	if (SelectedProfile.IsValid())
-	{
-		SelectedProfile->SetDeployWithUnrealPak(NewState == ESlateCheckBoxState::Checked);
-	}
-}
-
-
-ESlateCheckBoxState::Type SProjectLauncherDeployToDeviceSettings::HandleUnrealPakCheckBoxIsChecked( ) const
-{
-	ILauncherProfilePtr SelectedProfile = Model->GetSelectedProfile();
-
-	if (SelectedProfile.IsValid())
-	{
-		if (SelectedProfile->IsPackingWithUnrealPak())
-		{
-			return ESlateCheckBoxState::Checked;
-		}
-	}
-
-	return ESlateCheckBoxState::Unchecked;
-}
-
 
 #undef LOCTEXT_NAMESPACE
diff --git a/Engine/Source/Developer/ProjectLauncher/Private/Widgets/Deploy/SProjectLauncherDeployToDeviceSettings.h b/Engine/Source/Developer/ProjectLauncher/Private/Widgets/Deploy/SProjectLauncherDeployToDeviceSettings.h
--- a/Engine/Source/Developer/ProjectLauncher/Private/Widgets/Deploy/SProjectLauncherDeployToDeviceSettings.h
+++ b/Engine/Source/Developer/ProjectLauncher/Private/Widgets/Deploy/SProjectLauncherDeployToDeviceSettings.h
@@ -26,12 +26,6 @@ public:
 
 private:
 
-	// Callback for check state changes of the 'UnrealPak' check box.
-	void HandleUnrealPakCheckBoxCheckStateChanged( ESlateCheckBoxState::Type NewState );
-
-	// Callback for determining the checked state of the 'UnrealPak' check box.
-	ESlateCheckBoxState::Type HandleUnrealPakCheckBoxIsChecked( ) const;
-
 private:
 
 	// Holds a pointer to the data model.
