Made the add profile button stand out more so it's easier to see.

diff --git a/Engine/Source/Developer/ProjectLauncher/Private/Widgets/SProjectLauncher.cpp b/Engine/Source/Developer/ProjectLauncher/Private/Widgets/SProjectLauncher.cpp
--- a/Engine/Source/Developer/ProjectLauncher/Private/Widgets/SProjectLauncher.cpp
+++ b/Engine/Source/Developer/ProjectLauncher/Private/Widgets/SProjectLauncher.cpp
@@ -183,7 +183,7 @@ void SProjectLauncher::Construct( const FArguments& InArgs, const TSharedRef<SDo
 								[
 									SNew(SImage)
 									.Image(FCoreStyle::Get().GetBrush("EditableComboBox.Add"))
-									.ColorAndOpacity(FSlateColor::UseForeground())
+									.ColorAndOpacity(FSlateColor(FLinearColor::White))
 								]
 							]
 						]
