removed Generate DSym button
UE-3306
#ios

diff --git a/Engine/Source/Developer/ProjectLauncher/Private/Widgets/Build/SProjectLauncherBuildPage.cpp b/Engine/Source/Developer/ProjectLauncher/Private/Widgets/Build/SProjectLauncherBuildPage.cpp
--- a/Engine/Source/Developer/ProjectLauncher/Private/Widgets/Build/SProjectLauncherBuildPage.cpp
+++ b/Engine/Source/Developer/ProjectLauncher/Private/Widgets/Build/SProjectLauncherBuildPage.cpp
@@ -97,14 +97,14 @@ void SProjectLauncherBuildPage::Construct( const FArguments& InArgs, const FProj
                 [
                     SNew(SVerticalBox)
 
-                    + SVerticalBox::Slot()
+/*                    + SVerticalBox::Slot()
                     .AutoHeight()
                     [
                         SNew(SButton)
                         .Text(LOCTEXT("GenDSYMText", "Generate DSYM"))
                         .IsEnabled( this, &SProjectLauncherBuildPage::HandleGenDSYMButtonEnabled )
                         .OnClicked( this, &SProjectLauncherBuildPage::HandleGenDSYMClicked )
-                    ]
+                    ]*/
 					+ SVerticalBox::Slot()
 					.AutoHeight()
 					[
