Added support for compression in to project launcher and packaging settings in the editor.

diff --git a/Engine/Source/Developer/LauncherServices/Private/Launcher/LauncherWorker.cpp b/Engine/Source/Developer/LauncherServices/Private/Launcher/LauncherWorker.cpp
--- a/Engine/Source/Developer/LauncherServices/Private/Launcher/LauncherWorker.cpp
+++ b/Engine/Source/Developer/LauncherServices/Private/Launcher/LauncherWorker.cpp
@@ -424,6 +424,11 @@ FString FLauncherWorker::CreateUATCommand( const ILauncherProfileRef& InProfile,
 		UATCommand += TEXT(" -iterativecooking");
 	}
 
+	if ( InProfile->IsCompressed() )
+	{
+		UATCommand += TEXT(" -compressed");
+	}
+
 	// stage/package/deploy
 	if (InProfile->GetDeploymentMode() != ELauncherProfileDeploymentModes::DoNotDeploy)
 	{
diff --git a/Engine/Source/Developer/LauncherServices/Private/Profiles/LauncherProfile.h b/Engine/Source/Developer/LauncherServices/Private/Profiles/LauncherProfile.h
--- a/Engine/Source/Developer/LauncherServices/Private/Profiles/LauncherProfile.h
+++ b/Engine/Source/Developer/LauncherServices/Private/Profiles/LauncherProfile.h
@@ -414,6 +414,11 @@ public:
 		return false;
 	}
 
+	virtual bool IsCompressed() const override
+	{
+		return Compressed;
+	}
+
 	virtual bool IsCookingUnversioned( ) const override
 	{
 		return CookUnversioned;
@@ -604,6 +609,7 @@ public:
 		CookOptions = FString();
 		CookIncremental = false;
 		CookUnversioned = true;
+		Compressed = false;
 		CookedCultures.Reset();
 		CookedCultures.Add(I18N.GetCurrentCulture()->GetName());
 		CookedMaps.Reset();
@@ -780,6 +786,16 @@ public:
 		}
 	}
 
+	virtual void SetCompressed( bool Enabled ) override
+	{
+		if (Compressed != Enabled)
+		{
+			Compressed = Enabled;
+
+			Validate();
+		}
+	}
+
 	virtual void SetIncrementalDeploying( bool Incremental ) override
 	{
 		if (DeployIncremental != Incremental)
@@ -1125,6 +1141,9 @@ private:
 	// Holds a flag indicating whether UAT should be built
 	bool BuildUAT;
 
+	// Generate compressed content
+	bool Compressed;
+
 	// Holds a flag indicating whether only modified content should be cooked.
 	bool CookIncremental;
 
diff --git a/Engine/Source/Developer/LauncherServices/Public/Interfaces/ILauncherProfile.h b/Engine/Source/Developer/LauncherServices/Public/Interfaces/ILauncherProfile.h
--- a/Engine/Source/Developer/LauncherServices/Public/Interfaces/ILauncherProfile.h
+++ b/Engine/Source/Developer/LauncherServices/Public/Interfaces/ILauncherProfile.h
@@ -578,6 +578,14 @@ public:
 	 */
 	virtual bool IsCookingIncrementally( ) const = 0;
 
+	/**
+	 * Checks if compression is enabled
+	 *
+	 * @return true if compression is enabled
+	 * @see SetCompressed
+	 */
+	virtual bool IsCompressed( ) const = 0;
+
 	/**
 	 * Checks whether unversioned cooking is enabled.
 	 *
@@ -825,6 +833,14 @@ public:
 	 */
 	virtual void SetIncrementalCooking( bool Incremental ) = 0;
 
+	/**
+	 * Sets Compression.
+	 *
+	 * @param Enable compression
+	 * @see IsCompressed
+	 */
+	virtual void SetCompressed( bool Enable ) = 0;
+
 	/**
 	 * Sets incremental deploying.
 	 *
diff --git a/Engine/Source/Developer/ProjectLauncher/Private/Widgets/Cook/SProjectLauncherCookByTheBookSettings.cpp b/Engine/Source/Developer/ProjectLauncher/Private/Widgets/Cook/SProjectLauncherCookByTheBookSettings.cpp
--- a/Engine/Source/Developer/ProjectLauncher/Private/Widgets/Cook/SProjectLauncherCookByTheBookSettings.cpp
+++ b/Engine/Source/Developer/ProjectLauncher/Private/Widgets/Cook/SProjectLauncherCookByTheBookSettings.cpp
@@ -333,20 +333,36 @@ TSharedRef<SWidget> SProjectLauncherCookByTheBookSettings::MakeComplexWidget()
 				SNew(SVerticalBox)
 
 				+ SVerticalBox::Slot()
-				.AutoHeight()
-				[
-					// incremental cook check box
-					SNew(SCheckBox)
-					.IsChecked(this, &SProjectLauncherCookByTheBookSettings::HandleIncrementalCheckBoxIsChecked)
-					.OnCheckStateChanged(this, &SProjectLauncherCookByTheBookSettings::HandleIncrementalCheckBoxCheckStateChanged)
-					.Padding(FMargin(4.0f, 0.0f))
-					.ToolTipText(LOCTEXT("IncrementalCheckBoxTooltip", "If checked, only modified content will be cooked, resulting in much faster cooking times. It is recommended to enable this option whenever possible."))
-					.Content()
+					.AutoHeight()
 					[
-						SNew(STextBlock)
-						.Text(LOCTEXT("IncrementalCheckBoxText", "Only cook modified content"))
+						// incremental cook check box
+						SNew(SCheckBox)
+						.IsChecked(this, &SProjectLauncherCookByTheBookSettings::HandleIncrementalCheckBoxIsChecked)
+						.OnCheckStateChanged(this, &SProjectLauncherCookByTheBookSettings::HandleIncrementalCheckBoxCheckStateChanged)
+						.Padding(FMargin(4.0f, 0.0f))
+						.ToolTipText(LOCTEXT("IncrementalCheckBoxTooltip", "If checked, only modified content will be cooked, resulting in much faster cooking times. It is recommended to enable this option whenever possible."))
+						.Content()
+						[
+							SNew(STextBlock)
+							.Text(LOCTEXT("IncrementalCheckBoxText", "Only cook modified content"))
+						]
+					]
+
+				+ SVerticalBox::Slot()
+					.AutoHeight()
+					[
+						// incremental cook check box
+						SNew(SCheckBox)
+						.IsChecked(this, &SProjectLauncherCookByTheBookSettings::HandleCompressedCheckBoxIsChecked)
+						.OnCheckStateChanged(this, &SProjectLauncherCookByTheBookSettings::HandleCompressedCheckBoxCheckStateChanged)
+						.Padding(FMargin(4.0f, 0.0f))
+						.ToolTipText(LOCTEXT("CompressedCheckboxToolTip", "If checked, content will be generated compressed.  These will be smaller but potentially take longer to load"))
+						.Content()
+						[
+							SNew(STextBlock)
+							.Text(LOCTEXT("CompressedCheckBoxText", "Compress content"))
+						]
 					]
-				]
 
 				+ SVerticalBox::Slot()
 					.AutoHeight()
@@ -873,7 +889,36 @@ ECheckBoxState SProjectLauncherCookByTheBookSettings::HandleIncrementalCheckBoxI
 
 	if (SelectedProfile.IsValid())
 	{
-		if (SelectedProfile->IsCookingIncrementally())
+		if (SelectedProfile->IsCompressed())
+		{
+			return ECheckBoxState::Checked;
+		}
+	}
+
+	return ECheckBoxState::Unchecked;
+}
+
+
+
+
+void SProjectLauncherCookByTheBookSettings::HandleCompressedCheckBoxCheckStateChanged( ECheckBoxState NewState )
+{
+	ILauncherProfilePtr SelectedProfile = Model->GetSelectedProfile();
+
+	if (SelectedProfile.IsValid())
+	{
+		SelectedProfile->SetCompressed(NewState == ECheckBoxState::Checked);
+	}
+}
+
+
+ECheckBoxState SProjectLauncherCookByTheBookSettings::HandleCompressedCheckBoxIsChecked( ) const
+{
+	ILauncherProfilePtr SelectedProfile = Model->GetSelectedProfile();
+
+	if (SelectedProfile.IsValid())
+	{
+		if (SelectedProfile->IsCompressed())
 		{
 			return ECheckBoxState::Checked;
 		}
diff --git a/Engine/Source/Developer/ProjectLauncher/Private/Widgets/Cook/SProjectLauncherCookByTheBookSettings.h b/Engine/Source/Developer/ProjectLauncher/Private/Widgets/Cook/SProjectLauncherCookByTheBookSettings.h
--- a/Engine/Source/Developer/ProjectLauncher/Private/Widgets/Cook/SProjectLauncherCookByTheBookSettings.h
+++ b/Engine/Source/Developer/ProjectLauncher/Private/Widgets/Cook/SProjectLauncherCookByTheBookSettings.h
@@ -73,6 +73,10 @@ private:
 	/** Handles determining the checked state of the 'Incremental' check box. */
 	ECheckBoxState HandleIncrementalCheckBoxIsChecked( ) const;
 
+	void HandleCompressedCheckBoxCheckStateChanged( ECheckBoxState NewState );
+
+	ECheckBoxState HandleCompressedCheckBoxIsChecked( ) const;
+	
 	/** Handles generating a row widget in the map list view. */
 	TSharedRef<ITableRow> HandleMapListViewGenerateRow( TSharedPtr<FString> InItem, const TSharedRef<STableViewBase>& OwnerTable );
 
diff --git a/Engine/Source/Editor/MainFrame/Private/Frame/MainFrameActions.cpp b/Engine/Source/Editor/MainFrame/Private/Frame/MainFrameActions.cpp
--- a/Engine/Source/Editor/MainFrame/Private/Frame/MainFrameActions.cpp
+++ b/Engine/Source/Editor/MainFrame/Private/Frame/MainFrameActions.cpp
@@ -491,6 +491,11 @@ void FMainFrameActionCallbacks::PackageProject( const FName InPlatformInfoName )
 		OptionalParams += TEXT(" -clean");
 	}
 
+	if ( PackagingSettings->bCompressed )
+	{
+		OptionalParams += TEXT(" -compressed");
+	}
+
 	if ( PackagingSettings->bCookAll )
 	{
 		OptionalParams += TEXT(" -CookAll");
diff --git a/Engine/Source/Editor/UnrealEd/Classes/Settings/ProjectPackagingSettings.h b/Engine/Source/Editor/UnrealEd/Classes/Settings/ProjectPackagingSettings.h
--- a/Engine/Source/Editor/UnrealEd/Classes/Settings/ProjectPackagingSettings.h
+++ b/Engine/Source/Editor/UnrealEd/Classes/Settings/ProjectPackagingSettings.h
@@ -119,6 +119,13 @@ public:
 	 */
 	UPROPERTY(config, EditAnywhere, Category=Packaging, AdvancedDisplay, meta=(DisplayName="Cook only maps (this only affects cookall)"))
 	bool bCookMapsOnly;
+
+
+	/**
+	 * Create compressed cooked packages (decreased deployment size)
+	 */
+	UPROPERTY(config, EditAnywhere, Category=Packaging, AdvancedDisplay, meta=(DisplayName="Create compressed cooked packages"))
+	bool bCompressed;
 	
 	/**
 	 * List of maps to include when no other map list is specified on commandline
diff --git a/Engine/Source/Programs/AutomationTool/Scripts/CookCommand.Automation.cs b/Engine/Source/Programs/AutomationTool/Scripts/CookCommand.Automation.cs
--- a/Engine/Source/Programs/AutomationTool/Scripts/CookCommand.Automation.cs
+++ b/Engine/Source/Programs/AutomationTool/Scripts/CookCommand.Automation.cs
@@ -179,7 +179,12 @@ public partial class Project : CommandUtils
                 {
                     CommandletParams += " -dlcname=" + Params.DLCName;
                 }
-                
+                // if we are not going to pak but we specified compressed then compress in the cooker ;)
+                // otherwise compress the pak files
+                if (!Params.Pak && !Params.SkipPak && Params.Compressed)
+                {
+                    CommandletParams += " -compressed";
+                }
                 if (Params.HasAdditionalCookerOptions)
                 {
                     string FormatedAdditionalCookerParams = Params.AdditionalCookerOptions.TrimStart(new char[] { '\"', ' ' }).TrimEnd(new char[] { '\"', ' ' });
