Fixed issue with incrimental cook checkbox also highlighting the compression check box.

diff --git a/Engine/Source/Developer/ProjectLauncher/Private/Widgets/Cook/SProjectLauncherCookByTheBookSettings.cpp b/Engine/Source/Developer/ProjectLauncher/Private/Widgets/Cook/SProjectLauncherCookByTheBookSettings.cpp
--- a/Engine/Source/Developer/ProjectLauncher/Private/Widgets/Cook/SProjectLauncherCookByTheBookSettings.cpp
+++ b/Engine/Source/Developer/ProjectLauncher/Private/Widgets/Cook/SProjectLauncherCookByTheBookSettings.cpp
@@ -889,7 +889,7 @@ ECheckBoxState SProjectLauncherCookByTheBookSettings::HandleIncrementalCheckBoxI
 
 	if (SelectedProfile.IsValid())
 	{
-		if (SelectedProfile->IsCompressed())
+		if (SelectedProfile->IsCookingIncrementally())
 		{
 			return ECheckBoxState::Checked;
 		}
