Added support for unrealpak generating patches.
Exposed patching dlc and release creation in ufe.

diff --git a/Engine/Source/Developer/LauncherServices/Private/Launcher/LauncherWorker.cpp b/Engine/Source/Developer/LauncherServices/Private/Launcher/LauncherWorker.cpp
--- a/Engine/Source/Developer/LauncherServices/Private/Launcher/LauncherWorker.cpp
+++ b/Engine/Source/Developer/LauncherServices/Private/Launcher/LauncherWorker.cpp
@@ -372,6 +372,35 @@ FString FLauncherWorker::CreateUATCommand( const ILauncherProfileRef& InProfile,
 				UATCommand += TEXT(" -pak");
 			}
 
+			if ( InProfile->IsCreatingReleaseVersion() )
+			{
+				UATCommand += TEXT(" -createreleaseversion=");
+				UATCommand += InProfile->GetCreateReleaseVersionName();
+				
+			}
+
+			if ( InProfile->IsCreatingDLC() )
+			{
+				UATCommand += TEXT(" -dlcname=");
+				UATCommand += InProfile->GetDLCName();
+			}
+
+			if ( InProfile->IsGeneratingPatch() )
+			{
+				UATCommand += TEXT(" -generatepatch");
+			}
+
+			if ( InProfile->IsGeneratingPatch() || 
+				InProfile->IsCreatingReleaseVersion() || 
+				InProfile->IsCreatingDLC() )
+			{
+				if ( InProfile->GetBasedOnReleaseVersionName().IsEmpty() == false )
+				{
+					UATCommand += TEXT(" -basedonreleaseversion=");
+					UATCommand += InProfile->GetBasedOnReleaseVersionName();
+				}
+			}
+
 			FCommandDesc Desc;
 			FText Command = FText::Format(LOCTEXT("LauncherCookDesc", "Cook content for {0}"), FText::FromString(Platforms.RightChop(1)));
 			Desc.Name = "Cook Task";
@@ -424,6 +453,7 @@ FString FLauncherWorker::CreateUATCommand( const ILauncherProfileRef& InProfile,
 		UATCommand += TEXT(" -iterativecooking");
 	}
 
+
 	if ( InProfile->IsCompressed() )
 	{
 		UATCommand += TEXT(" -compressed");
diff --git a/Engine/Source/Developer/LauncherServices/Private/Profiles/LauncherProfile.h b/Engine/Source/Developer/LauncherServices/Private/Profiles/LauncherProfile.h
--- a/Engine/Source/Developer/LauncherServices/Private/Profiles/LauncherProfile.h
+++ b/Engine/Source/Developer/LauncherServices/Private/Profiles/LauncherProfile.h
@@ -3,6 +3,9 @@
 #pragma once
 
 #define LAUNCHERSERVICES_ADDEDINCREMENTALDEPLOYVERSION 11
+#define LAUNCHERSERVICES_ADDEDPATCHSOURCECONTENTPATH 12
+#define LAUNCHERSERVICES_ADDEDRELEASEVERSION 13
+#define LAUNCHERSERVICES_REMOVEDPATCHSOURCECONTENTPATH 14
 
 /**
 * Implements a simple profile which controls the desired output of the Launcher for simple
@@ -273,6 +276,64 @@ public:
 		return DeployedDeviceGroup;
 	}
 
+	virtual bool IsGeneratingPatch() const override
+	{
+		return GeneratePatch;
+	}
+
+
+	virtual bool IsCreatingDLC() const override
+	{
+		return CreateDLC;
+	}
+
+	virtual FString GetDLCName() const override
+	{
+		return DLCName;
+	}
+
+	virtual void SetCreateDLC(bool InBuildDLC) override
+	{
+		CreateDLC = InBuildDLC;
+	}
+	virtual void SetDLCName(const FString& InDLCName) override
+	{
+		DLCName = InDLCName;
+	}
+
+	virtual bool IsCreatingReleaseVersion() const override
+	{
+		return CreateReleaseVersion;
+	}
+
+	virtual void SetCreateReleaseVersion(bool InCreateReleaseVersion) override
+	{
+		CreateReleaseVersion = InCreateReleaseVersion;
+	}
+
+	virtual FString GetCreateReleaseVersionName() const override
+	{
+		return CreateReleaseVersionName;
+	}
+
+	virtual void SetCreateReleaseVersionName(const FString& InCreateReleaseVersionName) override
+	{
+		CreateReleaseVersionName = InCreateReleaseVersionName;
+	}
+
+
+	virtual FString GetBasedOnReleaseVersionName() const override
+	{
+		return BasedOnReleaseVersionName;
+	}
+
+	virtual void SetBasedOnReleaseVersionName(const FString& InBasedOnReleaseVersionName) override
+	{
+		BasedOnReleaseVersionName = InBasedOnReleaseVersionName;
+	}
+
+
+
 	virtual ELauncherProfileDeploymentModes::Type GetDeploymentMode( ) const override
 	{
 		return DeploymentMode;
@@ -542,6 +603,36 @@ public:
 			Archive << DeployIncremental;
 		}
 
+		if ( Version >= LAUNCHERSERVICES_REMOVEDPATCHSOURCECONTENTPATH )
+		{
+			Archive << GeneratePatch;
+		}
+		else if ( Version >= LAUNCHERSERVICES_ADDEDPATCHSOURCECONTENTPATH)
+		{
+			FString Temp;
+			Archive << Temp;
+			Archive << GeneratePatch;
+		}
+		else if ( Archive.IsLoading() )
+		{
+			GeneratePatch = false;
+		}
+
+		if ( Version >= LAUNCHERSERVICES_ADDEDRELEASEVERSION )
+		{
+			Archive << CreateReleaseVersion;
+			Archive << CreateReleaseVersionName;
+			Archive << BasedOnReleaseVersionName;
+			
+			Archive << CreateDLC;
+			Archive << DLCName;
+		}
+		else if ( Archive.IsLoading() )
+		{
+			CreateReleaseVersion = false;
+			CreateDLC = false;
+		}
+
 		DefaultLaunchRole->Serialize(Archive);
 
 		// serialize launch roles
@@ -631,6 +722,10 @@ public:
 		HideFileServerWindow = false;
 		DeployIncremental = false;
 
+		
+		GeneratePatch = false;
+		CreateDLC = false;
+
 		// default launch settings
 		LaunchMode = ELauncherProfileLaunchModes::DefaultRole;
 		DefaultLaunchRole->SetCommandLine(FString());
@@ -928,6 +1023,11 @@ public:
 		}
 	}
 
+	virtual void SetGeneratePatch( bool InGeneratePatch ) override
+	{
+		GeneratePatch = InGeneratePatch;
+	}
+
 	virtual bool SupportsEngineMaps( ) const override
 	{
 		return false;
@@ -1030,6 +1130,12 @@ protected:
 			ValidationErrors.Add(ELauncherProfileValidationErrors::UnversionedAndIncrimental);
 		}
 
+
+		if ( IsGeneratingPatch() && (CookMode != ELauncherProfileCookModes::ByTheBook) )
+		{
+			ValidationErrors.Add(ELauncherProfileValidationErrors::GeneratingPatchesCanOnlyRunFromByTheBookCookMode);
+		}
+
 		// Launch: when launching, all devices that the build is launched on must have content cooked for their platform
 		if ((LaunchMode != ELauncherProfileLaunchModes::DoNotLaunch) && (CookMode != ELauncherProfileCookModes::OnTheFly || CookMode != ELauncherProfileCookModes::OnTheFlyInEditor))
 		{
@@ -1168,6 +1274,28 @@ private:
 	// Holds a flag indicating whether content should be packaged with UnrealPak.
 	bool DeployWithUnrealPak;
 
+	// create a release version of the content (this can be used to base dlc / patches from)
+	bool CreateReleaseVersion;
+
+	// name of the release version
+	FString CreateReleaseVersionName;
+
+	// name of the release version to base this dlc / patch on
+	FString BasedOnReleaseVersionName;
+
+	// This build generate a patch based on some source content seealso PatchSourceContentPath
+	bool GeneratePatch;
+
+	// This build will cook content for dlc See also DLCName
+	bool CreateDLC;
+
+	// name of the dlc we are going to build (the name of the dlc plugin)
+	FString DLCName;
+
+	
+
+
+
 	// Holds a flag indicating whether to use incremental deployment
 	bool DeployIncremental;
 
diff --git a/Engine/Source/Developer/LauncherServices/Public/Interfaces/ILauncherProfile.h b/Engine/Source/Developer/LauncherServices/Public/Interfaces/ILauncherProfile.h
--- a/Engine/Source/Developer/LauncherServices/Public/Interfaces/ILauncherProfile.h
+++ b/Engine/Source/Developer/LauncherServices/Public/Interfaces/ILauncherProfile.h
@@ -159,6 +159,9 @@ namespace ELauncherProfileValidationErrors
 
 		/** The profile has unversioned and incrimental specified these are not compatible together */
 		UnversionedAndIncrimental,
+
+		/** generating patch requires cook by the book mode*/
+		GeneratingPatchesCanOnlyRunFromByTheBookCookMode,
 	};
 }
 
@@ -553,7 +556,13 @@ public:
      * @see SetTimeout
      */
     virtual uint32 GetTimeout() const = 0;
-    
+	
+	/**
+	 * Are we going to generate a patch (Source content patch needs to be specified)
+	 * @Seealso GetPatchSourceContentPath	 * @Seealso SetPatchSourceContentPath
+	 */
+	virtual bool IsGeneratingPatch() const = 0;
+
 	/**
 	 * Checks whether the game should be built.
 	 *
@@ -809,6 +818,39 @@ public:
 	 */
 	virtual void SetDeploymentMode( ELauncherProfileDeploymentModes::Type Mode ) = 0;
 
+	/**
+	 * Creating a release version of the cooked content 
+	 */
+	virtual bool IsCreatingReleaseVersion() const = 0;
+
+	virtual void SetCreateReleaseVersion(bool InCreateReleaseVersion) = 0;
+
+	virtual FString GetCreateReleaseVersionName() const = 0;
+
+	virtual void SetCreateReleaseVersionName(const FString& InCreateReleaseVersionName) = 0;
+
+	virtual FString GetBasedOnReleaseVersionName() const = 0;
+
+	virtual void SetBasedOnReleaseVersionName(const FString& InBasedOnReleaseVersion) = 0;
+
+
+	/**
+	 * Sets if we are going to generate a patch 
+	 * 
+	 * @param InShouldGeneratePatch enable generating patch
+	 * @seealso IsGeneratingPatch
+	 */
+	virtual void SetGeneratePatch( bool InShouldGeneratePatch ) = 0;
+
+	virtual bool IsCreatingDLC() const = 0;
+
+	virtual void SetCreateDLC(bool InBuildDLC) = 0;
+
+	virtual FString GetDLCName() const = 0;
+
+	virtual void SetDLCName(const FString& InDLCName) = 0;
+
+
     /**
      * Sets the cook on the fly close mode
      *
@@ -927,3 +969,4 @@ public:
 	/** Virtual destructor. */
 	virtual ~ILauncherProfile( ) { }
 };
+
diff --git a/Engine/Source/Developer/ProjectLauncher/Private/Widgets/Cook/SProjectLauncherCookByTheBookSettings.cpp b/Engine/Source/Developer/ProjectLauncher/Private/Widgets/Cook/SProjectLauncherCookByTheBookSettings.cpp
--- a/Engine/Source/Developer/ProjectLauncher/Private/Widgets/Cook/SProjectLauncherCookByTheBookSettings.cpp
+++ b/Engine/Source/Developer/ProjectLauncher/Private/Widgets/Cook/SProjectLauncherCookByTheBookSettings.cpp
@@ -320,6 +320,171 @@ TSharedRef<SWidget> SProjectLauncherCookByTheBookSettings::MakeComplexWidget()
 			]
 		]
 
+
+
+
+
+		+ SVerticalBox::Slot()
+			.AutoHeight()
+			.Padding(0.0f, 8.0f, 0.0f, 0.0f)
+			[
+				SNew(SExpandableArea)
+				.AreaTitle(LOCTEXT("PatchingAreaTitle", "Release / DLC / Patching Settings"))
+				.InitiallyCollapsed(true)
+				.Padding(8.0f)
+				.BodyContent()
+				[
+					SNew(SVerticalBox)
+
+
+					//////////////////////////////////////////////////////////////////////////
+					// create release version options
+					+ SVerticalBox::Slot()
+					.AutoHeight()
+					.Padding(0.0f, 4.0f, 0.0f, 0.0f)
+					[
+						// unreal pak check box
+						SNew(SCheckBox)
+						.IsChecked(this, &SProjectLauncherCookByTheBookSettings::HandleCreateReleaseVersionCheckBoxIsChecked)
+						.OnCheckStateChanged(this, &SProjectLauncherCookByTheBookSettings::HandleCreateReleaseVersionCheckBoxCheckStateChanged)
+						.Padding(FMargin(4.0f, 0.0f))
+						.ToolTipText(LOCTEXT("CreateReleaseVersionCheckBoxTooltip", "Create a release version of the game for distribution."))
+						.Content()
+						[
+							SNew(STextBlock)
+							.Text(LOCTEXT("CreateReleaseVersionBoxText", "Create a release version of the game for distribution."))
+						]
+					]
+
+					+ SVerticalBox::Slot()
+						.AutoHeight()
+						.Padding(0.0f, 8.0f, 0.0f, 0.0f)
+						[
+							SNew(SProjectLauncherFormLabel)
+							.LabelText(LOCTEXT("CreateReleaseVersionTextBoxLabel", "Name of the new release to create."))
+						]
+
+					+ SVerticalBox::Slot()
+						.AutoHeight()
+						.Padding(0.0, 4.0, 0.0, 0.0)
+						[
+							SNew(SHorizontalBox)
+
+							+ SHorizontalBox::Slot()
+							.FillWidth(1.0)
+							.Padding(0.0, 0.0, 0.0, 3.0)
+							[
+								// repository path text box
+								SNew(SEditableTextBox)
+								.ToolTipText(LOCTEXT("CreateReleaseVersionTextBoxTooltip", "Name of the new release to create."))
+								.Text(this, &SProjectLauncherCookByTheBookSettings::HandleCreateReleaseVersionNameTextBlockText)
+								.OnTextCommitted(this, &SProjectLauncherCookByTheBookSettings::HandleCreateReleaseVersionNameCommitted)
+							]
+						]
+
+
+					+ SVerticalBox::Slot()
+						.AutoHeight()
+						.Padding(0.0f, 8.0f, 0.0f, 0.0f)
+						[
+							SNew(SProjectLauncherFormLabel)
+							.LabelText(LOCTEXT("BasedOnReleaseVersionTextBoxLabel", "Base this release on a previous release, applys to new creating new releases / dlc / patches."))
+						]
+
+					+ SVerticalBox::Slot()
+						.AutoHeight()
+						.Padding(0.0, 4.0, 0.0, 0.0)
+						[
+							SNew(SHorizontalBox)
+
+							+ SHorizontalBox::Slot()
+							.FillWidth(1.0)
+							.Padding(0.0, 0.0, 0.0, 3.0)
+							[
+								// repository path text box
+								SNew(SEditableTextBox)
+								.ToolTipText(LOCTEXT("BasedOnReleaseVersionTextBoxTooltip", "Release version to base the next release / dlc / patch on."))
+								.Text(this, &SProjectLauncherCookByTheBookSettings::HandleBasedOnReleaseVersionNameTextBlockText)
+								.OnTextCommitted(this, &SProjectLauncherCookByTheBookSettings::HandleBasedOnReleaseVersionNameCommitted)
+							]
+						]
+
+
+					// end create release version
+					//////////////////////////////////////////////////////////////////////////
+					// generate patch params
+					+ SVerticalBox::Slot()
+						.AutoHeight()
+						.Padding(0.0f, 4.0f, 0.0f, 0.0f)
+						[
+							// unreal pak check box
+							SNew(SCheckBox)
+							.IsChecked(this, &SProjectLauncherCookByTheBookSettings::HandleGeneratePatchCheckBoxIsChecked)
+							.OnCheckStateChanged(this, &SProjectLauncherCookByTheBookSettings::HandleGeneratePatchCheckBoxCheckStateChanged)
+							.Padding(FMargin(4.0f, 0.0f))
+							.ToolTipText(LOCTEXT("GeneratePatchCheckBoxTooltip", "If checked, content will be diffed against source content and only changed files will be included in new pak"))
+							.Content()
+							[
+								SNew(STextBlock)
+								.Text(LOCTEXT("GeneratePatchCheckBoxText", "Generate patch"))
+							]
+						]
+
+					// end generate patch options
+					//////////////////////////////////////////////////////////////////////////
+					// generate dlc options
+					+ SVerticalBox::Slot()
+						.AutoHeight()
+						.Padding(0.0f, 4.0f, 0.0f, 0.0f)
+						[
+							// unreal pak check box
+							SNew(SCheckBox)
+							.IsChecked(this, &SProjectLauncherCookByTheBookSettings::HandleBuildDLCCheckBoxIsChecked)
+							.OnCheckStateChanged(this, &SProjectLauncherCookByTheBookSettings::HandleBuildDLCCheckBoxCheckStateChanged)
+							.Padding(FMargin(4.0f, 0.0f))
+							.ToolTipText(LOCTEXT("BuildDLCCheckBoxTooltip", "If checked, dlc will be built without the content released with the original game."))
+							.Content()
+							[
+								SNew(STextBlock)
+								.Text(LOCTEXT("BuildDLCCheckBoxText", "Build dlc"))
+							]
+						]
+
+
+					+ SVerticalBox::Slot()
+						.AutoHeight()
+						.Padding(0.0f, 8.0f, 0.0f, 0.0f)
+						[
+							SNew(SProjectLauncherFormLabel)
+							.LabelText(LOCTEXT("DLCNameTextBoxLabel", "Name of the dlc to build."))
+						]
+
+					+ SVerticalBox::Slot()
+						.AutoHeight()
+						.Padding(0.0, 4.0, 0.0, 0.0)
+						[
+							SNew(SHorizontalBox)
+
+							+ SHorizontalBox::Slot()
+							.FillWidth(1.0)
+							.Padding(0.0, 0.0, 0.0, 3.0)
+							[
+								// repository path text box
+								SNew(SEditableTextBox)
+								.ToolTipText(LOCTEXT("DLCNameTextBoxTooltip", "Name of DLC to build."))
+								.Text(this, &SProjectLauncherCookByTheBookSettings::HandleDLCNameTextBlockText)
+								.OnTextCommitted(this, &SProjectLauncherCookByTheBookSettings::HandleDLCNameCommitted)
+							]
+						]
+					// end generate dlc 
+					//////////////////////////////////////////////////////////////////////////
+			]
+		]
+
+		
+
+
+
 	+ SVerticalBox::Slot()
 		.AutoHeight()
 		.Padding(0.0f, 8.0f, 0.0f, 0.0f)
@@ -349,9 +514,9 @@ TSharedRef<SWidget> SProjectLauncherCookByTheBookSettings::MakeComplexWidget()
 					]
 
 				+ SVerticalBox::Slot()
+					.Padding(0.0f, 4.0f, 0.0f, 0.0f)
 					.AutoHeight()
 					[
-						// incremental cook check box
 						SNew(SCheckBox)
 						.IsChecked(this, &SProjectLauncherCookByTheBookSettings::HandleCompressedCheckBoxIsChecked)
 						.OnCheckStateChanged(this, &SProjectLauncherCookByTheBookSettings::HandleCompressedCheckBoxCheckStateChanged)
@@ -1116,4 +1281,165 @@ ECheckBoxState SProjectLauncherCookByTheBookSettings::HandleUnrealPakCheckBoxIsC
 	return ECheckBoxState::Unchecked;
 }
 
+// Callback for check state changes of the 'UnrealPak' check box.
+void SProjectLauncherCookByTheBookSettings::HandleGeneratePatchCheckBoxCheckStateChanged( ECheckBoxState NewState )
+{
+	ILauncherProfilePtr SelectedProfile = Model->GetSelectedProfile();
+
+	if (SelectedProfile.IsValid())
+	{
+		SelectedProfile->SetGeneratePatch(NewState == ECheckBoxState::Checked);
+	}
+}
+
+// Callback for determining the checked state of the 'UnrealPak' check box.
+ECheckBoxState SProjectLauncherCookByTheBookSettings::HandleGeneratePatchCheckBoxIsChecked( ) const
+{
+	ILauncherProfilePtr SelectedProfile = Model->GetSelectedProfile();
+
+	if (SelectedProfile.IsValid())
+	{
+		if (SelectedProfile->IsGeneratingPatch())
+		{
+			return ECheckBoxState::Checked;
+		}
+	}
+
+	return ECheckBoxState::Unchecked;
+}
+
+
+
+
+// Callback for check state changes of the 'UnrealPak' check box.
+void SProjectLauncherCookByTheBookSettings::HandleCreateReleaseVersionCheckBoxCheckStateChanged( ECheckBoxState NewState )
+{
+	ILauncherProfilePtr SelectedProfile = Model->GetSelectedProfile();
+
+	if (SelectedProfile.IsValid())
+	{
+		SelectedProfile->SetCreateReleaseVersion(NewState == ECheckBoxState::Checked);
+	}
+}
+
+// Callback for determining the checked state of the 'UnrealPak' check box.
+ECheckBoxState SProjectLauncherCookByTheBookSettings::HandleCreateReleaseVersionCheckBoxIsChecked( ) const
+{
+	ILauncherProfilePtr SelectedProfile = Model->GetSelectedProfile();
+
+	if (SelectedProfile.IsValid())
+	{
+		if (SelectedProfile->IsCreatingReleaseVersion())
+		{
+			return ECheckBoxState::Checked;
+		}
+	}
+
+	return ECheckBoxState::Unchecked;
+}
+
+FText SProjectLauncherCookByTheBookSettings::HandleCreateReleaseVersionNameTextBlockText() const
+{
+	ILauncherProfilePtr SelectedProfile = Model->GetSelectedProfile();
+
+	FText result;
+
+	if (SelectedProfile.IsValid())
+	{
+		result = FText::FromString(SelectedProfile->GetCreateReleaseVersionName());
+	}
+
+	return result;
+}
+
+void SProjectLauncherCookByTheBookSettings::HandleCreateReleaseVersionNameCommitted(const FText& NewText, ETextCommit::Type CommitType)
+{
+	ILauncherProfilePtr SelectedProfile = Model->GetSelectedProfile();
+
+	if (SelectedProfile.IsValid())
+	{
+		SelectedProfile->SetCreateReleaseVersionName(NewText.ToString());
+	}
+
+}
+
+
+FText SProjectLauncherCookByTheBookSettings::HandleBasedOnReleaseVersionNameTextBlockText() const
+{
+	ILauncherProfilePtr SelectedProfile = Model->GetSelectedProfile();
+
+	FText result;
+
+	if (SelectedProfile.IsValid())
+	{
+		result = FText::FromString(SelectedProfile->GetBasedOnReleaseVersionName());
+	}
+
+	return result;
+}
+
+void SProjectLauncherCookByTheBookSettings::HandleBasedOnReleaseVersionNameCommitted(const FText& NewText, ETextCommit::Type CommitType)
+{
+	ILauncherProfilePtr SelectedProfile = Model->GetSelectedProfile();
+
+	if (SelectedProfile.IsValid())
+	{
+		SelectedProfile->SetBasedOnReleaseVersionName(NewText.ToString());
+	}
+
+}
+
+
+FText SProjectLauncherCookByTheBookSettings::HandleDLCNameTextBlockText() const
+{
+	ILauncherProfilePtr SelectedProfile = Model->GetSelectedProfile();
+
+	FText result;
+
+	if (SelectedProfile.IsValid())
+	{
+		result = FText::FromString(SelectedProfile->GetDLCName());
+	}
+
+	return result;
+}
+
+void SProjectLauncherCookByTheBookSettings::HandleDLCNameCommitted(const FText& NewText, ETextCommit::Type CommitType)
+{
+	ILauncherProfilePtr SelectedProfile = Model->GetSelectedProfile();
+
+	if (SelectedProfile.IsValid())
+	{
+		SelectedProfile->SetDLCName(NewText.ToString());
+	}
+}
+
+// Callback for check state changes of the 'UnrealPak' check box.
+void SProjectLauncherCookByTheBookSettings::HandleBuildDLCCheckBoxCheckStateChanged( ECheckBoxState NewState ) 
+{
+	ILauncherProfilePtr SelectedProfile = Model->GetSelectedProfile();
+
+	if (SelectedProfile.IsValid())
+	{
+		SelectedProfile->SetCreateDLC(NewState == ECheckBoxState::Checked);
+	}
+}
+
+// Callback for determining the checked state of the 'UnrealPak' check box.
+ECheckBoxState SProjectLauncherCookByTheBookSettings::HandleBuildDLCCheckBoxIsChecked() const
+{
+	ILauncherProfilePtr SelectedProfile = Model->GetSelectedProfile();
+
+	if (SelectedProfile.IsValid())
+	{
+		if (SelectedProfile->IsCreatingDLC())
+		{
+			return ECheckBoxState::Checked;
+		}
+	}
+
+	return ECheckBoxState::Unchecked;
+}
+
+
 #undef LOCTEXT_NAMESPACE
diff --git a/Engine/Source/Developer/ProjectLauncher/Private/Widgets/Cook/SProjectLauncherCookByTheBookSettings.h b/Engine/Source/Developer/ProjectLauncher/Private/Widgets/Cook/SProjectLauncherCookByTheBookSettings.h
--- a/Engine/Source/Developer/ProjectLauncher/Private/Widgets/Cook/SProjectLauncherCookByTheBookSettings.h
+++ b/Engine/Source/Developer/ProjectLauncher/Private/Widgets/Cook/SProjectLauncherCookByTheBookSettings.h
@@ -125,6 +125,41 @@ private:
 	// Callback for determining the checked state of the 'UnrealPak' check box.
 	ECheckBoxState HandleUnrealPakCheckBoxIsChecked( ) const;
 
+
+	//////////////////////////////////////////////////////////////////////////
+	// creating release version related functions
+	// Callback for check state changes of the 'GeneratePatch' check box.
+	void HandleCreateReleaseVersionCheckBoxCheckStateChanged( ECheckBoxState NewState );
+	// Callback for determining the checked state of the 'GeneratePatch' check box.
+	ECheckBoxState HandleCreateReleaseVersionCheckBoxIsChecked( ) const;
+
+	void HandleCreateReleaseVersionNameCommitted(const FText& NewText, ETextCommit::Type CommitType);
+	FText HandleCreateReleaseVersionNameTextBlockText() const;
+
+	void HandleBasedOnReleaseVersionNameCommitted(const FText& NewText, ETextCommit::Type CommitType);
+	FText HandleBasedOnReleaseVersionNameTextBlockText() const;
+
+	//////////////////////////////////////////////////////////////////////////
+	// patch generation related functions
+	// Callback for check state changes of the 'GeneratePatch' check box.
+	void HandleGeneratePatchCheckBoxCheckStateChanged( ECheckBoxState NewState );
+	// Callback for determining the checked state of the 'GeneratePatch' check box.
+	ECheckBoxState HandleGeneratePatchCheckBoxIsChecked( ) const;
+
+	// Callback for changing patchSourceContent path (should be the path of a pak file)
+	FText HandlePatchSourceContentPathTextBlockText() const;
+	// Callback for getting the PatchSourceContentPath
+	void HandlePatchSourceContentPathCommitted(const FText& NewText, ETextCommit::Type CommitType);
+
+	//////////////////////////////////////////////////////////////////////////
+	// dlc check box related functions 
+	void HandleBuildDLCCheckBoxCheckStateChanged( ECheckBoxState NewState );
+	ECheckBoxState HandleBuildDLCCheckBoxIsChecked() const;
+	void HandleDLCNameCommitted(const FText& NewText, ETextCommit::Type CommitType);
+	FText HandleDLCNameTextBlockText() const;
+	
+
+
 	/** creates the complex widget. */
 	TSharedRef<SWidget> MakeComplexWidget();
 
@@ -133,6 +168,12 @@ private:
 
 private:
 
+	/** Textbox which holds the PatchSourceContentPath */
+	TSharedPtr<SEditableTextBox> PatchSourceContentPath;
+
+	/** Textbox which holds the DLCBasedOnReleaseVersion */
+	TSharedPtr<SEditableTextBox> DLCBasedOnReleaseVersionName;
+
 	/** Holds the culture list. */
 	TArray<TSharedPtr<FString> > CultureList;
 
diff --git a/Engine/Source/Programs/AutomationTool/ProjectParams.cs b/Engine/Source/Programs/AutomationTool/ProjectParams.cs
--- a/Engine/Source/Programs/AutomationTool/ProjectParams.cs
+++ b/Engine/Source/Programs/AutomationTool/ProjectParams.cs
@@ -198,6 +198,7 @@ namespace AutomationTool
             this.CulturesToCook = InParams.CulturesToCook;
             this.BasedOnReleaseVersion = InParams.BasedOnReleaseVersion;
             this.CreateReleaseVersion = InParams.CreateReleaseVersion;
+            this.GeneratePatch = InParams.GeneratePatch;
             this.DLCName = InParams.DLCName;
             this.NewCook = InParams.NewCook;
             this.OldCook = InParams.OldCook;
@@ -341,6 +342,7 @@ namespace AutomationTool
             string AdditionalCookerOptions = null,
             string BasedOnReleaseVersion = null,
             string CreateReleaseVersion = null,
+            bool? GeneratePatch = null,
             string DLCName = null,
             bool? NewCook = null,
             bool? OldCook = null,
@@ -442,6 +444,7 @@ namespace AutomationTool
             this.OldCook = GetParamValueIfNotSpecified(Command, OldCook, this.OldCook, "OldCook");
             this.CreateReleaseVersion = ParseParamValueIfNotSpecified(Command, CreateReleaseVersion, "createreleaseversion", String.Empty);
             this.BasedOnReleaseVersion = ParseParamValueIfNotSpecified(Command, BasedOnReleaseVersion, "basedonreleaseversion", String.Empty);
+            this.GeneratePatch = GetParamValueIfNotSpecified(Command, GeneratePatch, this.GeneratePatch, "GeneratePatch");
             this.AdditionalCookerOptions = ParseParamValueIfNotSpecified(Command, AdditionalCookerOptions, "AdditionalCookerOptions", String.Empty);
             this.DLCName = ParseParamValueIfNotSpecified(Command, DLCName, "DLCName", String.Empty);
 			this.SkipCook = GetParamValueIfNotSpecified(Command, SkipCook, this.SkipCook, "skipcook");
@@ -1009,6 +1012,13 @@ namespace AutomationTool
         /// </summary>
         public string BasedOnReleaseVersion;
 
+        /// <summary>
+        /// Are we generating a patch, generate a patch from a previously released version of the game (use CreateReleaseVersion to create a release). 
+        /// this requires BasedOnReleaseVersion
+        /// see also CreateReleaseVersion, BasedOnReleaseVersion
+        /// </summary>
+        public bool GeneratePatch;
+
         /// <summary>
         /// Name of dlc to cook and package (if this paramter is supplied cooks the dlc and packages it into the dlc directory)
         /// </summary>
@@ -1039,6 +1049,7 @@ namespace AutomationTool
         /// </summary>
         public bool UnversionedCookedContent;
 
+
 		/// <summary>
 		/// Cook: Uses the iterative cooking, command line: -iterativecooking or -iterate
 		/// </summary>
@@ -1711,6 +1722,14 @@ namespace AutomationTool
 		}
 		private string ProjectBinariesPath;
 
+        /// <summary>
+        /// True if we are generating a patch
+        /// </summary>
+        public bool IsGeneratingPatch
+        {
+            get { return GeneratePatch; }
+        }
+
 		/// <summary>
 		/// Filename of the target game exe (or program exe).
 		/// </summary>
@@ -1893,9 +1912,24 @@ namespace AutomationTool
 
 			if (Cook && CookOnTheFly)
 			{
-				throw new AutomationException("Don't use both -cook and -cookonthefly.");
+				throw new AutomationException("Can't use both -cook and -cookonthefly.");
 			}
 
+            if ((IsGeneratingPatch || HasDLCName) && !HasBasedOnReleaseVersion)
+            {
+                throw new AutomationException("Require based on release version to build patches or dlc");
+            }
+
+            if (HasCreateReleaseVersion && (IsGeneratingPatch || HasDLCName))
+            {
+                throw new AutomationException("Can't create a release version at the same time as creating dlc.");
+            }
+
+            if (HasBasedOnReleaseVersion && (IterativeCooking || IterativeDeploy))
+            {
+                throw new AutomationException("Can't use iterative cooking / deploy on dlc or patching or creating a release");
+            }
+
             /*if (Compressed && !Pak)
             {
                 throw new AutomationException("-compressed can only be used with -pak");
@@ -1934,6 +1968,7 @@ namespace AutomationTool
 				CommandUtils.Log("CookOnTheFly={0}", CookOnTheFly);
                 CommandUtils.Log("CookOnTheFlyStreaming={0}", CookOnTheFlyStreaming);
                 CommandUtils.Log("UnversionedCookedContent={0}", UnversionedCookedContent);
+                CommandUtils.Log("GeneratePatch={0}", GeneratePatch);
                 CommandUtils.Log("CreateReleaseVersion={0}", CreateReleaseVersion);
                 CommandUtils.Log("BasedOnReleaseVersion={0}", BasedOnReleaseVersion);
                 CommandUtils.Log("DLCName={0}", DLCName);
diff --git a/Engine/Source/Programs/AutomationTool/Scripts/CookCommand.Automation.cs b/Engine/Source/Programs/AutomationTool/Scripts/CookCommand.Automation.cs
--- a/Engine/Source/Programs/AutomationTool/Scripts/CookCommand.Automation.cs
+++ b/Engine/Source/Programs/AutomationTool/Scripts/CookCommand.Automation.cs
@@ -171,14 +171,16 @@ public partial class Project : CommandUtils
                 {
                     CommandletParams += " -createreleaseversion=" + Params.CreateReleaseVersion;
                 }
-                if (Params.HasBasedOnReleaseVersion)
-                {
-                    CommandletParams += " -basedonreleaseversion=" + Params.BasedOnReleaseVersion;
-                }
                 if (Params.HasDLCName)
                 {
                     CommandletParams += " -dlcname=" + Params.DLCName;
                 }
+                // don't include the based on release version unless we are cooking dlc or creating a new release version
+                // in this case the based on release version is used in packaging
+                if (Params.HasBasedOnReleaseVersion && (Params.HasDLCName || Params.HasCreateReleaseVersion))
+                {
+                    CommandletParams += " -basedonreleaseversion=" + Params.BasedOnReleaseVersion;
+                }
                 // if we are not going to pak but we specified compressed then compress in the cooker ;)
                 // otherwise compress the pak files
                 if (!Params.Pak && !Params.SkipPak && Params.Compressed)
@@ -188,6 +190,7 @@ public partial class Project : CommandUtils
                 if (Params.HasAdditionalCookerOptions)
                 {
                     string FormatedAdditionalCookerParams = Params.AdditionalCookerOptions.TrimStart(new char[] { '\"', ' ' }).TrimEnd(new char[] { '\"', ' ' });
+                    CommandletParams += " ";
                     CommandletParams += FormatedAdditionalCookerParams;
                 }
                 CookCommandlet(Params.RawProjectPath, Params.UE4Exe, Maps, Dirs, InternationalizationPreset, Cultures, CombineCommandletParams(PlatformsToCook.ToArray()), CommandletParams);
diff --git a/Engine/Source/Programs/AutomationTool/Scripts/CopyBuildToStagingDirectory.Automation.cs b/Engine/Source/Programs/AutomationTool/Scripts/CopyBuildToStagingDirectory.Automation.cs
--- a/Engine/Source/Programs/AutomationTool/Scripts/CopyBuildToStagingDirectory.Automation.cs
+++ b/Engine/Source/Programs/AutomationTool/Scripts/CopyBuildToStagingDirectory.Automation.cs
@@ -92,7 +92,7 @@ public partial class Project : CommandUtils
 		return Result;
 	}
 
-	static public void RunUnrealPak(Dictionary<string, string> UnrealPakResponseFile, string OutputLocation, string EncryptionKeys, string PakOrderFileLocation, string PlatformOptions, bool Compressed)
+    static public void RunUnrealPak(Dictionary<string, string> UnrealPakResponseFile, string OutputLocation, string EncryptionKeys, string PakOrderFileLocation, string PlatformOptions, bool Compressed, String PatchSourceContentPath)
 	{
 		if (UnrealPakResponseFile.Count < 1)
 		{
@@ -119,6 +119,10 @@ public partial class Project : CommandUtils
 		{
 			CmdLine += " -UTF8Output";
 		}
+        if (!String.IsNullOrEmpty(PatchSourceContentPath))
+        {
+            CmdLine += " -generatepatch=" + CommandUtils.MakePathSafeToUseWithCommandLine(PatchSourceContentPath);
+        }
 		CmdLine += PlatformOptions;
 		RunAndLog(CmdEnv, UnrealPakExe, CmdLine, Options: ERunOptions.Default | ERunOptions.UTF8Output);
 		Log("UnrealPak Done *******");
@@ -622,6 +626,14 @@ public partial class Project : CommandUtils
 		return UnrealPakResponseFile;
 	}
 
+
+    private static string GetReleasePakFilePath(DeploymentContext SC, ProjectParams Params, string ReleaseVersion, string PakName)
+    {
+        string ReleaseVersionPath = CombinePaths(SC.ProjectRoot, "Releases", ReleaseVersion, SC.StageTargetPlatform.GetCookPlatform(Params.DedicatedServer, false, Params.CookFlavor), PakName);
+        return ReleaseVersionPath;
+    }
+
+
 	/// <summary>
 	/// Creates a pak file using response file.
 	/// </summary>
@@ -631,7 +643,12 @@ public partial class Project : CommandUtils
 	/// <param name="PakName"></param>
 	private static void CreatePak(ProjectParams Params, DeploymentContext SC, Dictionary<string, string> UnrealPakResponseFile, string PakName)
 	{
-		var OutputRelativeLocation = CombinePaths(SC.RelativeProjectRootForStage, "Content/Paks/", PakName + "-" + SC.FinalCookPlatform + ".pak");
+        string PostFix = "";
+        if (Params.IsGeneratingPatch)
+        {
+            PostFix += "_P";
+        }
+		var OutputRelativeLocation = CombinePaths(SC.RelativeProjectRootForStage, "Content/Paks/", PakName + "-" + SC.FinalCookPlatform + PostFix + ".pak");
 		if (SC.StageTargetPlatform.DeployLowerCaseFilenames(true))
 		{
 			OutputRelativeLocation = OutputRelativeLocation.ToLowerInvariant();
@@ -653,7 +670,7 @@ public partial class Project : CommandUtils
 		{
 			// Check to see if we have an existing pak file we can use
 
-			var SourceOutputRelativeLocation = CombinePaths(SC.RelativeProjectRootForStage, "Content/Paks/", PakName + "-" + SC.CookPlatform + ".pak");
+			var SourceOutputRelativeLocation = CombinePaths(SC.RelativeProjectRootForStage, "Content/Paks/", PakName + "-" + SC.CookPlatform + PostFix + ".pak");
 			if (SC.CookSourcePlatform.DeployLowerCaseFilenames(true))
 			{
 				SourceOutputRelativeLocation = SourceOutputRelativeLocation.ToLowerInvariant();
@@ -677,11 +694,40 @@ public partial class Project : CommandUtils
 			}
 		}
 
+        string PatchSourceContentPath = null;
+        if (Params.HasBasedOnReleaseVersion && Params.IsGeneratingPatch)
+        {
+            // don't include the post fix in this filename because we are looking for the source pak path
+            string PakFilename = PakName + "-" + SC.FinalCookPlatform + ".pak";
+            PatchSourceContentPath = GetReleasePakFilePath(SC, Params, Params.BasedOnReleaseVersion, PakFilename);
+            
+
+            /*PatchSourceContentPath = Params.BasedOnReleaseVersion;
+            if (PatchSourceContentPath.EndsWith(PakName) == false)
+            {
+                string Temp = Path.GetDirectoryName(PatchSourceContentPath) + PakName;
+                if ( File.Exists(Temp))
+                {
+                    PatchSourceContentPath = Temp;
+                }
+            }*/
+        }
+
 		if (!bCopiedExistingPak)
 		{
-			RunUnrealPak(UnrealPakResponseFile, OutputLocation, Params.SignPak, PakOrderFileLocation, SC.StageTargetPlatform.GetPlatformPakCommandLine(), Params.Compressed);
+			RunUnrealPak(UnrealPakResponseFile, OutputLocation, Params.SignPak, PakOrderFileLocation, SC.StageTargetPlatform.GetPlatformPakCommandLine(), Params.Compressed, PatchSourceContentPath );
 		}
 
+
+        if (Params.HasCreateReleaseVersion)
+        {
+            // copy the created pak to the release version directory we might need this later if we want to generate patches
+            //string ReleaseVersionPath = CombinePaths( SC.ProjectRoot, "Releases", Params.CreateReleaseVersion, SC.StageTargetPlatform.GetCookPlatform(Params.DedicatedServer, false, Params.CookFlavor), Path.GetFileName(OutputLocation) );
+            string ReleaseVersionPath = GetReleasePakFilePath(SC, Params, Params.CreateReleaseVersion, Path.GetFileName(OutputLocation));
+
+            File.Copy(OutputLocation, ReleaseVersionPath);
+        }
+
 		if (Params.CreateChunkInstall)
 		{
 			var RegEx = new Regex("pakchunk(\\d+)", RegexOptions.IgnoreCase | RegexOptions.Compiled);
diff --git a/Engine/Source/Programs/UnrealPak/Private/UnrealPak.cpp b/Engine/Source/Programs/UnrealPak/Private/UnrealPak.cpp
--- a/Engine/Source/Programs/UnrealPak/Private/UnrealPak.cpp
+++ b/Engine/Source/Programs/UnrealPak/Private/UnrealPak.cpp
@@ -17,10 +17,14 @@ struct FPakCommandLineParameters
 	FPakCommandLineParameters()
 		: CompressionBlockSize(64*1024)
 		, FileSystemBlockSize(0)
+		, GeneratePatch(false)
 	{}
 
 	int32  CompressionBlockSize;
 	int64  FileSystemBlockSize;
+	bool   GeneratePatch;
+	FString SourcePatchPakFilename;
+	FString SourcePatchDiffDirectory;
 };
 
 struct FPakEntryPair
@@ -408,18 +412,12 @@ void ProcessCommandLine(int32 ArgC, TCHAR* ArgV[], TArray<FPakInputPair>& Entrie
 
 	if (FParse::Value(FCommandLine::Get(), TEXT("-create="), ResponseFile))
 	{
-		bool bCompress = false;
-		bool bEncrypt = false;
 		TArray<FString> Lines;
 
-		if (FParse::Param(FCommandLine::Get(), TEXT("compress")))
-		{
-			bCompress = true;
-		}
-		if (FParse::Param(FCommandLine::Get(), TEXT("encrypt")))
-		{
-			bEncrypt = true;
-		}
+		CmdLineParameters.GeneratePatch = FParse::Value(FCommandLine::Get(), TEXT("-generatepatch="), CmdLineParameters.SourcePatchPakFilename);
+
+		bool bCompress = FParse::Param(FCommandLine::Get(), TEXT("compress"));
+		bool bEncrypt = FParse::Param(FCommandLine::Get(), TEXT("encrypt"));
 
 		if (IFileManager::Get().DirectoryExists(*ResponseFile))
 		{
@@ -943,7 +941,7 @@ bool ListFilesInPak(const TCHAR * InPakFilename)
 	}
 }
 
-bool ExtractFilesFromPak(const TCHAR* InPakFilename, const TCHAR* InDestPath)
+bool ExtractFilesFromPak(const TCHAR* InPakFilename, const TCHAR* InDestPath, bool bUseMountPoint = false)
 {
 	FPakFile PakFile(InPakFilename, FParse::Param(FCommandLine::Get(), TEXT("signed")));
 	if (PakFile.IsValid())
@@ -957,6 +955,8 @@ bool ExtractFilesFromPak(const TCHAR* InPakFilename, const TCHAR* InDestPath)
 		int32 ErrorCount = 0;
 		int32 FileCount = 0;
 
+		FString PakMountPoint = bUseMountPoint ? PakFile.GetMountPoint().Replace( TEXT("../../../"), TEXT("")) : TEXT("");
+
 		for (FPakFile::FFileIterator It(PakFile); It; ++It, ++FileCount)
 		{
 			const FPakEntry& Entry = It.Info();
@@ -966,7 +966,8 @@ bool ExtractFilesFromPak(const TCHAR* InPakFilename, const TCHAR* InDestPath)
 			EntryInfo.Serialize(PakReader, PakFile.GetInfo().Version);
 			if (EntryInfo == Entry)
 			{
-				FString DestFilename(DestPath / It.Filename());
+				FString DestFilename(DestPath / PakMountPoint /  It.Filename());
+
 				TAutoPtr<FArchive> FileHandle(IFileManager::Get().CreateFileWriter(*DestFilename));
 				if (FileHandle.IsValid())
 				{
@@ -1006,6 +1007,84 @@ bool ExtractFilesFromPak(const TCHAR* InPakFilename, const TCHAR* InDestPath)
 	}
 }
 
+bool GenerateHashForFile( FString Filename, uint8 FileHash[16])
+{
+	FArchive* File = IFileManager::Get().CreateFileReader(*Filename);
+
+	if ( File == NULL )
+		return false;
+
+	uint64 TotalSize = File->TotalSize();
+
+	uint8* ByteBuffer = new uint8[TotalSize];
+
+	File->Serialize(ByteBuffer, TotalSize);
+
+	delete File;
+	File = NULL;
+
+	FMD5 FileHasher;
+	FileHasher.Update(ByteBuffer, TotalSize);
+
+	delete[] ByteBuffer;
+
+	FileHasher.Final(FileHash);
+
+	return true;
+
+
+	// uint8 DestFileHash[20];
+
+}
+
+
+void RemoveIdenticalFiles( TArray<FPakInputPair>& FilesToPak, const FString& SourceDirectory )
+{
+	for ( int I = FilesToPak.Num()-1; I >= 0; --I )
+	{
+		const auto& NewFile = FilesToPak[I]; 
+
+		FString SourceFilename = SourceDirectory / NewFile.Dest.Replace(TEXT("../../../"), TEXT(""));
+		int64 SourceTotalSize = IFileManager::Get().FileSize(*SourceFilename);
+
+		FString DestFilename = NewFile.Source;
+		int64 DestTotalSize = IFileManager::Get().FileSize(*DestFilename);
+		
+		if ( SourceTotalSize != DestTotalSize )
+		{
+			// file size doesn't match 
+			UE_LOG(LogPakFile, Display, TEXT("Source file size for %s %d bytes doesn't match %s %d bytes"), *SourceFilename, SourceTotalSize, *DestFilename, DestTotalSize);
+			continue;
+		}
+
+		uint8 SourceFileHash[16];
+		if ( GenerateHashForFile(SourceFilename, SourceFileHash) == false )
+		{
+			// file size doesn't match 
+			UE_LOG(LogPakFile, Display, TEXT("Source file size %s doesn't exist will be included in build"), *SourceFilename);
+			continue;
+		}
+		uint8 DestFileHash[16];
+		if ( GenerateHashForFile( DestFilename, DestFileHash ) == false )
+		{
+			// destination file was removed don't really care about it
+			UE_LOG(LogPakFile, Display, TEXT("File was removed from destination cooked content %s not included in patch"), *DestFilename);
+			continue;
+		}
+
+		int32 Diff = FMemory::Memcmp( SourceFileHash, DestFileHash, sizeof( DestFileHash ) );
+		if ( Diff != 0 )
+		{
+			UE_LOG(LogPakFile, Display, TEXT("Source file hash for %s doesn't match dest file hash %s and will be included in patch"), *SourceFilename, *DestFilename);
+			continue;
+		}
+
+		UE_LOG(LogPakFile, Display, TEXT("Source file %s matches dest file %s and will not be included in patch"), *SourceFilename, *DestFilename);
+		// remove fromt eh files to pak list
+		FilesToPak.RemoveAt(I);
+	}
+}
+
 /**
  * Application entry point
  * Params:
@@ -1089,10 +1168,33 @@ INT32_MAIN_INT32_ARGC_TCHAR_ARGV()
 			}
 			else
 			{
+				if ( CmdLineParameters.GeneratePatch )
+				{
+					FString OutputPath = FPaths::GetPath(PakFilename) / FString(TEXT("TempFiles"));
+
+					IFileManager::Get().DeleteDirectory(*OutputPath);
+
+					UE_LOG(LogPakFile, Display, TEXT("Generating patch from %s."), *CmdLineParameters.SourcePatchPakFilename );
+
+					if ( ExtractFilesFromPak( *CmdLineParameters.SourcePatchPakFilename, *OutputPath, true ) == false )
+					{
+						UE_LOG(LogPakFile, Error, TEXT("Unable to extract files from source pak file for patch") );
+					}
+					CmdLineParameters.SourcePatchDiffDirectory = OutputPath;
+				}
+
+
 				// Start collecting files
 				TArray<FPakInputPair> FilesToAdd;
 				CollectFilesToAdd(FilesToAdd, Entries, OrderMap);
 
+				if ( CmdLineParameters.GeneratePatch )
+				{
+					// if we are generating a patch here we remove files which are already shipped...
+					RemoveIdenticalFiles(FilesToAdd, CmdLineParameters.SourcePatchDiffDirectory);
+				}
+
+
 				Result = CreatePakFile(*PakFilename, FilesToAdd, CmdLineParameters) ? 0 : 1;
 			}
 		}
