Fixed Mac CIS compile regression

mono: Developer/ProjectLauncher/Private/Widgets/Cook/SProjectLauncherCookByTheBookSettings.h:161:56: error: extra qualification on member 'HandleDLCIncludeEngineContentCheckBoxIsChecked'

diff --git a/Engine/Source/Developer/ProjectLauncher/Private/Widgets/Cook/SProjectLauncherCookByTheBookSettings.h b/Engine/Source/Developer/ProjectLauncher/Private/Widgets/Cook/SProjectLauncherCookByTheBookSettings.h
--- a/Engine/Source/Developer/ProjectLauncher/Private/Widgets/Cook/SProjectLauncherCookByTheBookSettings.h
+++ b/Engine/Source/Developer/ProjectLauncher/Private/Widgets/Cook/SProjectLauncherCookByTheBookSettings.h
@@ -158,7 +158,7 @@ private:
 	void HandleDLCNameCommitted(const FText& NewText, ETextCommit::Type CommitType);
 	FText HandleDLCNameTextBlockText() const;
 	void HandleDLCIncludeEngineContentCheckBoxCheckStateChanged(ECheckBoxState NewState);
-	ECheckBoxState SProjectLauncherCookByTheBookSettings::HandleDLCIncludeEngineContentCheckBoxIsChecked() const;
+	ECheckBoxState HandleDLCIncludeEngineContentCheckBoxIsChecked() const;
 	
 
 
