Adding HTTP Chunk Installer Options to Editor Project Launcher/UnrealFrontend
Initial HTTPChunkInstaller patch support

diff --git a/Engine/Source/Developer/LauncherServices/Private/Launcher/LauncherWorker.cpp b/Engine/Source/Developer/LauncherServices/Private/Launcher/LauncherWorker.cpp
--- a/Engine/Source/Developer/LauncherServices/Private/Launcher/LauncherWorker.cpp
+++ b/Engine/Source/Developer/LauncherServices/Private/Launcher/LauncherWorker.cpp
@@ -405,6 +405,17 @@ FString FLauncherWorker::CreateUATCommand( const ILauncherProfileRef& InProfile,
 				}
 			}
 
+			if (InProfile->IsGeneratingChunks())
+			{
+				UATCommand += TEXT(" -manifests");
+			}
+
+			if (InProfile->IsGenerateHttpChunkData())
+			{
+				auto Cmd = FString::Printf(TEXT(" -createchunkinstall -chunkinstalldirectory=\"%s\" -chunkinstallversion=\"%s\""), *InProfile->GetHttpChunkDataDirectory(), *InProfile->GetHttpChunkDataReleaseName());
+				UATCommand += Cmd;
+			}
+
 			FCommandDesc Desc;
 			FText Command = FText::Format(LOCTEXT("LauncherCookDesc", "Cook content for {0}"), FText::FromString(Platforms.RightChop(1)));
 			Desc.Name = "Cook Task";
diff --git a/Engine/Source/Developer/LauncherServices/Private/LauncherServicesPrivatePCH.h b/Engine/Source/Developer/LauncherServices/Private/LauncherServicesPrivatePCH.h
--- a/Engine/Source/Developer/LauncherServices/Private/LauncherServicesPrivatePCH.h
+++ b/Engine/Source/Developer/LauncherServices/Private/LauncherServicesPrivatePCH.h
@@ -24,7 +24,7 @@
 #define LAUNCHERSERVICES_SIMPLEPROFILEVERSION 1
 
 /** Defines the launcher profile file version. */
-#define LAUNCHERSERVICES_PROFILEVERSION 14
+#define LAUNCHERSERVICES_PROFILEVERSION 16 
 #define LAUNCHERSERVICES_MINPROFILEVERSION 10
 
 
diff --git a/Engine/Source/Developer/LauncherServices/Private/Profiles/LauncherProfile.h b/Engine/Source/Developer/LauncherServices/Private/Profiles/LauncherProfile.h
--- a/Engine/Source/Developer/LauncherServices/Private/Profiles/LauncherProfile.h
+++ b/Engine/Source/Developer/LauncherServices/Private/Profiles/LauncherProfile.h
@@ -7,6 +7,7 @@
 #define LAUNCHERSERVICES_ADDEDRELEASEVERSION 13
 #define LAUNCHERSERVICES_REMOVEDPATCHSOURCECONTENTPATH 14
 #define LAUNCHERSERVICES_ADDEDDLCINCLUDEENGINECONTENT 15
+#define LAUNCHERSERVICES_ADDEDGENERATECHUNKS 16
 
 /**
 * Implements a simple profile which controls the desired output of the Launcher for simple
@@ -526,6 +527,26 @@ public:
 		return DeployWithUnrealPak;
 	}
 
+	virtual bool IsGeneratingChunks() const override
+	{
+		return bGenerateChunks;
+	}
+
+	virtual bool IsGenerateHttpChunkData() const override
+	{
+		return bGenerateHttpChunkData;
+	}
+
+	virtual FString GetHttpChunkDataDirectory() const override
+	{
+		return HttpChunkDataDirectory;
+	}
+
+	virtual FString GetHttpChunkDataReleaseName() const override
+	{
+		return HttpChunkDataReleaseName;
+	}
+
 	virtual bool IsValidForLaunch( ) override
 	{
 		return (ValidationErrors.Num() == 0);
@@ -654,6 +675,21 @@ public:
 			CreateDLC = false;
 		}
 
+		if (Version >= LAUNCHERSERVICES_ADDEDGENERATECHUNKS)
+		{
+			Archive << bGenerateChunks;
+			Archive << bGenerateHttpChunkData;
+			Archive << HttpChunkDataDirectory;
+			Archive << HttpChunkDataReleaseName;
+		}
+		else if (Archive.IsLoading())
+		{
+			bGenerateChunks = false;
+			bGenerateHttpChunkData = false;
+			HttpChunkDataDirectory = TEXT("");
+			HttpChunkDataReleaseName = TEXT("");
+		}
+
 		DefaultLaunchRole->Serialize(Archive);
 
 		// serialize launch roles
@@ -748,6 +784,11 @@ public:
 		CreateDLC = false;
 		DLCIncludeEngineContent = false;
 
+		bGenerateChunks = false;
+		bGenerateHttpChunkData = false;
+		HttpChunkDataDirectory = TEXT("");
+		HttpChunkDataReleaseName = TEXT("");
+
 		// default launch settings
 		LaunchMode = ELauncherProfileLaunchModes::DefaultRole;
 		DefaultLaunchRole->SetCommandLine(FString());
@@ -837,6 +878,42 @@ public:
 		}
 	}
 
+	virtual void SetGenerateChunks(bool bInGenerateChunks) override
+	{
+		if (bGenerateChunks != bInGenerateChunks)
+		{
+			bGenerateChunks = bInGenerateChunks;
+			Validate();
+		}
+	}
+
+	virtual void SetGenerateHttpChunkData(bool bInGenerateHttpChunkData) override
+	{
+		if (bGenerateHttpChunkData != bInGenerateHttpChunkData)
+		{
+			bGenerateHttpChunkData = bInGenerateHttpChunkData;
+			Validate();
+		}
+	}
+
+	virtual void SetHttpChunkDataDirectory(const FString& InHttpChunkDataDirectory) override
+	{
+		if (HttpChunkDataDirectory != InHttpChunkDataDirectory)
+		{
+			HttpChunkDataDirectory = InHttpChunkDataDirectory;
+			Validate();
+		}
+	}
+
+	virtual void SetHttpChunkDataReleaseName(const FString& InHttpChunkDataReleaseName) override
+	{
+		if (HttpChunkDataReleaseName != InHttpChunkDataReleaseName)
+		{
+			HttpChunkDataReleaseName = InHttpChunkDataReleaseName;
+			Validate();
+		}
+	}
+
 	virtual void SetDeployedDeviceGroup( const ILauncherDeviceGroupPtr& DeviceGroup ) override
 	{
 		if(DeployedDeviceGroup.IsValid())
@@ -1158,6 +1235,26 @@ protected:
 			ValidationErrors.Add(ELauncherProfileValidationErrors::GeneratingPatchesCanOnlyRunFromByTheBookCookMode);
 		}
 
+		if ( IsGeneratingChunks() && (CookMode != ELauncherProfileCookModes::ByTheBook) )
+		{
+			ValidationErrors.Add(ELauncherProfileValidationErrors::GeneratingChunksRequiresCookByTheBook);
+		}
+
+		if (IsGeneratingChunks() && !IsPackingWithUnrealPak())
+		{
+			ValidationErrors.Add(ELauncherProfileValidationErrors::GeneratingChunksRequiresUnrealPak);
+		}
+
+		if (IsGenerateHttpChunkData() && !IsGeneratingChunks())
+		{
+			ValidationErrors.Add(ELauncherProfileValidationErrors::GeneratingHttpChunkDataRequiresGeneratingChunks);
+		}
+
+		if (IsGenerateHttpChunkData() && (GetHttpChunkDataReleaseName().IsEmpty() || !FPaths::DirectoryExists(*GetHttpChunkDataDirectory())))
+		{
+			ValidationErrors.Add(ELauncherProfileValidationErrors::GeneratingHttpChunkDataRequiresValidDirectoryAndName);
+		}
+
 		// Launch: when launching, all devices that the build is launched on must have content cooked for their platform
 		if ((LaunchMode != ELauncherProfileLaunchModes::DoNotLaunch) && (CookMode != ELauncherProfileCookModes::OnTheFly || CookMode != ELauncherProfileCookModes::OnTheFlyInEditor))
 		{
@@ -1296,6 +1393,18 @@ private:
 	// Holds a flag indicating whether content should be packaged with UnrealPak.
 	bool DeployWithUnrealPak;
 
+	// Flag indicating if content should be split into chunks
+	bool bGenerateChunks;
+	
+	// Flag indicating if chunked content should be used to generate HTTPChunkInstall data
+	bool bGenerateHttpChunkData;
+	
+	// Where to store HTTPChunkInstall data
+	FString HttpChunkDataDirectory;
+	
+	// Version name of the HTTPChunkInstall data
+	FString HttpChunkDataReleaseName;
+
 	// create a release version of the content (this can be used to base dlc / patches from)
 	bool CreateReleaseVersion;
 
diff --git a/Engine/Source/Developer/LauncherServices/Public/Interfaces/ILauncherProfile.h b/Engine/Source/Developer/LauncherServices/Public/Interfaces/ILauncherProfile.h
--- a/Engine/Source/Developer/LauncherServices/Public/Interfaces/ILauncherProfile.h
+++ b/Engine/Source/Developer/LauncherServices/Public/Interfaces/ILauncherProfile.h
@@ -162,6 +162,18 @@ namespace ELauncherProfileValidationErrors
 
 		/** generating patch requires cook by the book mode*/
 		GeneratingPatchesCanOnlyRunFromByTheBookCookMode,
+
+		/** Generating Chunks requires cook by the book mode */
+		GeneratingChunksRequiresCookByTheBook,
+
+		/** Generating Chunks requires UnrealPak */
+		GeneratingChunksRequiresUnrealPak,
+
+		/** Generating http chunk install data requires generating chunks */
+		GeneratingHttpChunkDataRequiresGeneratingChunks,
+
+		/** Generating http chunk install data requires valid install directorys and release name */
+		GeneratingHttpChunkDataRequiresValidDirectoryAndName,
 	};
 }
 
@@ -635,6 +647,34 @@ public:
 	 */
 	virtual bool IsPackingWithUnrealPak( ) const = 0;
 
+	/**
+	 * Return whether packaging will generate chunk data.
+	 *
+	 * @return true if Chunks will be generated, false otherwise.
+	 */
+	virtual bool IsGeneratingChunks() const = 0;
+	
+	/**
+	 * Return whether packaging will use chunk data to generate http chunk install data.
+	 *
+	 * @return true if data will be generated, false otherwise.
+	 */
+	virtual bool IsGenerateHttpChunkData() const = 0;
+	
+	/**
+	 * Where generated http chunk install data will be stored.
+	 *
+	 * @return true if data will be generated, false otherwise.
+	 */
+	virtual FString GetHttpChunkDataDirectory() const = 0;
+	
+	/**
+	 * What name to tag the generated http chunk install data with.
+	 *
+	 * @return true if data will be generated, false otherwise.
+	 */
+	virtual FString GetHttpChunkDataReleaseName() const = 0;
+
 	/**
 	 * Checks whether the profile's selected project supports Engine maps.
 	 *
@@ -802,6 +842,31 @@ public:
 	 */
 	virtual void SetDeployWithUnrealPak( bool UseUnrealPak ) = 0;
 
+	/**
+	 * Set whether packaging will generate chunk data.
+	 *
+	 * @param true if Chunks should be generated, false otherwise.
+	 */
+	virtual void SetGenerateChunks(bool bGenerateChunks) = 0;
+	/**
+	 * Set whether packaging will use chunk data to generate http chunk install data.
+	 *
+	 * @param true if data should be generated, false otherwise.
+	 */
+	virtual void SetGenerateHttpChunkData(bool bGenerateHttpChunkData) = 0;
+	/**
+	 * Set where generated http chunk install data will be stored.
+	 *
+	 * @return the directory path to use.
+	 */	
+	virtual void SetHttpChunkDataDirectory(const FString& InHttpChunkDataDirectory ) = 0;
+	/**
+	 * Set what name to tag the generated http chunk install data with.
+	 *
+	 * @param the name to use.
+	 */	
+	virtual void SetHttpChunkDataReleaseName(const FString& InHttpChunkDataReleaseName ) = 0;
+
 	/**
 	 * Sets the device group to deploy to.
 	 *
diff --git a/Engine/Source/Developer/ProjectLauncher/Private/Widgets/Cook/SProjectLauncherCookByTheBookSettings.cpp b/Engine/Source/Developer/ProjectLauncher/Private/Widgets/Cook/SProjectLauncherCookByTheBookSettings.cpp
--- a/Engine/Source/Developer/ProjectLauncher/Private/Widgets/Cook/SProjectLauncherCookByTheBookSettings.cpp
+++ b/Engine/Source/Developer/ProjectLauncher/Private/Widgets/Cook/SProjectLauncherCookByTheBookSettings.cpp
@@ -580,6 +580,110 @@ TSharedRef<SWidget> SProjectLauncherCookByTheBookSettings::MakeComplexWidget()
 						]
 					]
 
+				+ SVerticalBox::Slot()
+					.AutoHeight()
+					.Padding(0.0f, 4.0f, 0.0f, 0.0f)
+					[
+						// unreal pak check box
+						SNew(SCheckBox)
+						.IsChecked(this, &SProjectLauncherCookByTheBookSettings::HandleGenerateChunksCheckBoxIsChecked)
+						.OnCheckStateChanged(this, &SProjectLauncherCookByTheBookSettings::HandleGenerateChunksCheckBoxCheckStateChanged)
+						.Padding(FMargin(4.0f, 0.0f))
+						.ToolTipText(LOCTEXT("GenerateChunksCheckBoxTooltip", "If checked, the content will be deployed as multiple UnrealPak files instead of many separate files."))
+						.Content()
+						[
+							SNew(STextBlock)
+							.Text(LOCTEXT("GenerateChunksCheckBoxText", "Generate Chunks"))
+						]
+					]
+
+				+ SVerticalBox::Slot()
+					.AutoHeight()
+					.Padding(0.0f, 4.0f, 0.0f, 0.0f)
+					[
+						SNew(SExpandableArea)
+						.AreaTitle(LOCTEXT("HttpChunkInstallSettingsAreaTitle", "Http Chunk Install Settings"))
+						.InitiallyCollapsed(true)
+						.Padding(FMargin(4.0f, 0.0f))
+						.BodyContent()
+						[
+							SNew(SVerticalBox)
+
+							+ SVerticalBox::Slot()
+								.AutoHeight()
+								.Padding(0.0f, 4.0f, 0.0f, 0.0f)
+								[
+									// unreal pak check box
+									SNew(SCheckBox)
+									.IsChecked(this, &SProjectLauncherCookByTheBookSettings::HandleHttpChunkInstallCheckBoxIsChecked)
+									.OnCheckStateChanged(this, &SProjectLauncherCookByTheBookSettings::HandleHttpChunkInstallCheckBoxCheckStateChanged)
+									.Padding(FMargin(4.0f, 0.0f))
+									.ToolTipText(LOCTEXT("HttpChunkInstallCheckBoxTooltip", "If checked, the content will be split into multiple paks and stored as data that can be downloaded."))
+									.Content()
+									[
+										SNew(STextBlock)
+										.Text(LOCTEXT("HttpChunkInstallCheckBoxText", "Create Http Chunk Install data"))
+									]
+								]
+
+							+ SVerticalBox::Slot()
+								.AutoHeight()
+								[
+									SNew(STextBlock)
+									.Text(LOCTEXT("HttpChunkInstallDataPathLabel", "Http Chunk Install Data Path:"))
+								]
+
+							+ SVerticalBox::Slot()
+								.AutoHeight()
+								.Padding(0.0, 4.0, 0.0, 0.0)
+								[
+									SNew(SHorizontalBox)
+
+									+ SHorizontalBox::Slot()
+									.FillWidth(1.0)
+									.Padding(0.0f, 4.0f, 0.0f, 0.0f)
+									[
+										// repository path text box
+										SAssignNew(HttpChunkInstallDirectoryTextBox, SEditableTextBox)
+										.OnTextCommitted(this, &SProjectLauncherCookByTheBookSettings::HandleHtppChunkInstallDirectoryTextCommitted)
+										.OnTextChanged(this, &SProjectLauncherCookByTheBookSettings::HandleHtppChunkInstallDirectoryTextChanged)
+									]
+
+									+ SHorizontalBox::Slot()
+										.AutoWidth()
+										.HAlign(HAlign_Right)
+										.Padding(4.0, 0.0, 0.0, 0.0)
+										[
+											// browse button
+											SNew(SButton)
+											.ContentPadding(FMargin(6.0, 2.0))
+											.IsEnabled(true)
+											.Text(LOCTEXT("BrowseButtonText", "Browse..."))
+											.ToolTipText(LOCTEXT("BrowseButtonToolTip", "Browse for the Http Chunk Install Data directory"))
+											.OnClicked(this, &SProjectLauncherCookByTheBookSettings::HandleHtppChunkInstallBrowseButtonClicked)
+										]
+								]
+				
+							+ SVerticalBox::Slot()
+								.AutoHeight()
+								.Padding(0.0f, 4.0f, 0.0f, 0.0f)
+								[
+									SNew(SProjectLauncherFormLabel)
+									.LabelText(LOCTEXT("HttpChunkInstallReleaseTextBoxLabel", "Http Chunk Install Release Name:"))
+								]
+
+							+ SVerticalBox::Slot()
+								.AutoHeight()
+								.Padding(0.0f, 4.0f, 0.0f, 0.0f)
+								[
+									// cooker command line options
+									SNew(SEditableTextBox)
+									.ToolTipText(LOCTEXT("HttpChunkInstallReleaseTextBoxTooltip", "Name of this version of the Http Chunk Install data."))
+									.Text(this, &SProjectLauncherCookByTheBookSettings::HandleHttpChunkInstallNameTextBlockText)
+									.OnTextCommitted(this, &SProjectLauncherCookByTheBookSettings::HandleHtppChunkInstallNameCommitted)
+								]
+						]
+					]
 				+ SVerticalBox::Slot()
 					.AutoHeight()
 					.Padding(0.0f, 12.0f, 0.0f, 0.0f)
@@ -1456,6 +1560,111 @@ ECheckBoxState SProjectLauncherCookByTheBookSettings::HandleBuildDLCCheckBoxIsCh
 	return ECheckBoxState::Unchecked;
 }
 
+FReply SProjectLauncherCookByTheBookSettings::HandleHtppChunkInstallBrowseButtonClicked()
+{
+	IDesktopPlatform* DesktopPlatform = FDesktopPlatformModule::Get();
+	if (DesktopPlatform)
+	{
+		TSharedPtr<SWindow> ParentWindow = FSlateApplication::Get().FindWidgetWindow(AsShared());
+		void* ParentWindowHandle = (ParentWindow.IsValid() && ParentWindow->GetNativeWindow().IsValid()) ? ParentWindow->GetNativeWindow()->GetOSWindowHandle() : nullptr;
+
+		FString FolderName;
+		const bool bFolderSelected = DesktopPlatform->OpenDirectoryDialog(
+			ParentWindowHandle,
+			LOCTEXT("RepositoryBrowseTitle", "Choose a repository location").ToString(),
+			HttpChunkInstallDirectoryTextBox->GetText().ToString(),
+			FolderName
+			);
+
+		if (bFolderSelected)
+		{
+			if (!FolderName.EndsWith(TEXT("/")))
+			{
+				FolderName += TEXT("/");
+			}
+
+			HttpChunkInstallDirectoryTextBox->SetText(FText::FromString(FolderName));
+			ILauncherProfilePtr SelectedProfile = Model->GetSelectedProfile();
+
+			if (SelectedProfile.IsValid())
+			{
+				SelectedProfile->SetHttpChunkDataDirectory(FolderName);
+			}
+		}
+	}
+
+	return FReply::Handled();
+}
+
+
+void SProjectLauncherCookByTheBookSettings::HandleHtppChunkInstallDirectoryTextChanged(const FText& InText)
+{
+	ILauncherProfilePtr SelectedProfile = Model->GetSelectedProfile();
+
+	if (SelectedProfile.IsValid())
+	{
+		SelectedProfile->SetHttpChunkDataDirectory(InText.ToString());
+	}
+}
+
+
+void SProjectLauncherCookByTheBookSettings::HandleHtppChunkInstallDirectoryTextCommitted(const FText& InText, ETextCommit::Type CommitInfo)
+{
+	if (CommitInfo == ETextCommit::OnEnter)
+	{
+		ILauncherProfilePtr SelectedProfile = Model->GetSelectedProfile();
+
+		if (SelectedProfile.IsValid())
+		{
+			SelectedProfile->SetHttpChunkDataDirectory(InText.ToString());
+		}
+	}
+}
+
+void SProjectLauncherCookByTheBookSettings::HandleHttpChunkInstallCheckBoxCheckStateChanged(ECheckBoxState NewState)
+{
+	ILauncherProfilePtr SelectedProfile = Model->GetSelectedProfile();
+	if (SelectedProfile.IsValid())
+	{
+		SelectedProfile->SetGenerateHttpChunkData(NewState == ECheckBoxState::Checked);
+	}
+}
+
+ECheckBoxState SProjectLauncherCookByTheBookSettings::HandleHttpChunkInstallCheckBoxIsChecked() const
+{
+	ILauncherProfilePtr SelectedProfile = Model->GetSelectedProfile();
+	if (SelectedProfile.IsValid())
+	{
+		return SelectedProfile->IsGenerateHttpChunkData() ? ECheckBoxState::Checked : ECheckBoxState::Unchecked;
+	}
+	return ECheckBoxState::Unchecked;
+}
+
+FText SProjectLauncherCookByTheBookSettings::HandleHttpChunkInstallNameTextBlockText() const
+{
+	ILauncherProfilePtr SelectedProfile = Model->GetSelectedProfile();
+
+	if (SelectedProfile.IsValid())
+	{
+		return FText::FromString(SelectedProfile->GetHttpChunkDataReleaseName());
+	}
+
+	return FText();
+}
+
+void SProjectLauncherCookByTheBookSettings::HandleHtppChunkInstallNameCommitted(const FText& NewText, ETextCommit::Type CommitType)
+{
+	if (CommitType == ETextCommit::OnEnter)
+	{
+		ILauncherProfilePtr SelectedProfile = Model->GetSelectedProfile();
+
+		if (SelectedProfile.IsValid())
+		{
+			SelectedProfile->SetHttpChunkDataReleaseName(NewText.ToString());
+		}
+	}
+}
+
 
 
 // Callback for check state changes of the 'UnrealPak' check box.
@@ -1485,12 +1694,24 @@ ECheckBoxState SProjectLauncherCookByTheBookSettings::HandleDLCIncludeEngineCont
 	return ECheckBoxState::Unchecked;
 }
 
+void SProjectLauncherCookByTheBookSettings::HandleGenerateChunksCheckBoxCheckStateChanged(ECheckBoxState NewState)
+{
+	ILauncherProfilePtr SelectedProfile = Model->GetSelectedProfile();
 
+	if (SelectedProfile.IsValid())
+	{
+		SelectedProfile->SetGenerateChunks(NewState == ECheckBoxState::Checked);
+	}
+}
 
-
-
-
-
-
+ECheckBoxState SProjectLauncherCookByTheBookSettings::HandleGenerateChunksCheckBoxIsChecked() const
+{
+	ILauncherProfilePtr SelectedProfile = Model->GetSelectedProfile();
+	if (SelectedProfile.IsValid())
+	{
+		return SelectedProfile->IsGeneratingChunks() ? ECheckBoxState::Checked : ECheckBoxState::Unchecked;
+	}
+	return ECheckBoxState::Unchecked;
+}
 
 #undef LOCTEXT_NAMESPACE
diff --git a/Engine/Source/Developer/ProjectLauncher/Private/Widgets/Cook/SProjectLauncherCookByTheBookSettings.h b/Engine/Source/Developer/ProjectLauncher/Private/Widgets/Cook/SProjectLauncherCookByTheBookSettings.h
--- a/Engine/Source/Developer/ProjectLauncher/Private/Widgets/Cook/SProjectLauncherCookByTheBookSettings.h
+++ b/Engine/Source/Developer/ProjectLauncher/Private/Widgets/Cook/SProjectLauncherCookByTheBookSettings.h
@@ -125,6 +125,9 @@ private:
 	// Callback for determining the checked state of the 'UnrealPak' check box.
 	ECheckBoxState HandleUnrealPakCheckBoxIsChecked( ) const;
 
+	void HandleGenerateChunksCheckBoxCheckStateChanged(ECheckBoxState NewState);
+	ECheckBoxState HandleGenerateChunksCheckBoxIsChecked() const;
+
 
 	//////////////////////////////////////////////////////////////////////////
 	// creating release version related functions
@@ -159,8 +162,16 @@ private:
 	FText HandleDLCNameTextBlockText() const;
 	void HandleDLCIncludeEngineContentCheckBoxCheckStateChanged(ECheckBoxState NewState);
 	ECheckBoxState HandleDLCIncludeEngineContentCheckBoxIsChecked() const;
-	
 
+	//////////////////////////////////////////////////////////////////////////
+	// Http Chunk Installer options & functions
+	FReply HandleHtppChunkInstallBrowseButtonClicked();
+	void HandleHtppChunkInstallDirectoryTextChanged(const FText& InText);
+	void HandleHtppChunkInstallDirectoryTextCommitted(const FText& InText, ETextCommit::Type CommitInfo);
+	void HandleHttpChunkInstallCheckBoxCheckStateChanged(ECheckBoxState NewState);
+	ECheckBoxState HandleHttpChunkInstallCheckBoxIsChecked() const;
+	FText HandleHttpChunkInstallNameTextBlockText() const;
+	void HandleHtppChunkInstallNameCommitted(const FText& NewText, ETextCommit::Type CommitType);
 
 	/** creates the complex widget. */
 	TSharedRef<SWidget> MakeComplexWidget();
@@ -193,4 +204,7 @@ private:
 
 	/** Holds the current 'Show maps' check box choice. */
 	EShowMapsChoices::Type ShowMapsChoice;
+
+	// Holds the Http Chunk Install directory path text box.
+	TSharedPtr<SEditableTextBox> HttpChunkInstallDirectoryTextBox;
 };
diff --git a/Engine/Source/Developer/ProjectLauncher/Private/Widgets/SProjectLauncherValidation.h b/Engine/Source/Developer/ProjectLauncher/Private/Widgets/SProjectLauncherValidation.h
--- a/Engine/Source/Developer/ProjectLauncher/Private/Widgets/SProjectLauncherValidation.h
+++ b/Engine/Source/Developer/ProjectLauncher/Private/Widgets/SProjectLauncherValidation.h
@@ -100,7 +100,25 @@ public:
 				[
 					MakeValidationMessage(TEXT("Icons.Error"), LOCTEXT("NoLaunchRoleDeviceAssigned", "One or more launch roles do not have a device assigned.").ToString(), ELauncherProfileValidationErrors::NoLaunchRoleDeviceAssigned)
 				]
-		 
+
+			+ SVerticalBox::Slot()
+				.AutoHeight()
+				[
+					MakeValidationMessage(TEXT("Icons.Error"), LOCTEXT("GeneratingChunksRequiresUnrealPak", "UnrealPak must be selected to Generate Chunks.").ToString(), ELauncherProfileValidationErrors::GeneratingChunksRequiresUnrealPak)
+				]
+			
+			+ SVerticalBox::Slot()
+				.AutoHeight()
+				[
+					MakeValidationMessage(TEXT("Icons.Error"), LOCTEXT("GeneratingHttpChunkDataRequiresGeneratingChunks", "Generate Chunks must be selected to Generate Http Chunk Install Data.").ToString(), ELauncherProfileValidationErrors::GeneratingHttpChunkDataRequiresGeneratingChunks)
+				]
+
+			+ SVerticalBox::Slot()
+				.AutoHeight()
+				[
+					MakeValidationMessage(TEXT("Icons.Error"), LOCTEXT("GeneratingHttpChunkDataRequiresValidDirectoryAndName", "Generating Http Chunk Install Data requires a valid directory and release name.").ToString(), ELauncherProfileValidationErrors::GeneratingHttpChunkDataRequiresValidDirectoryAndName)
+				]
+
 			+ SVerticalBox::Slot()
 				.AutoHeight()
 				[
diff --git a/Engine/Source/Programs/AutomationTool/Scripts/CopyBuildToStagingDirectory.Automation.cs b/Engine/Source/Programs/AutomationTool/Scripts/CopyBuildToStagingDirectory.Automation.cs
--- a/Engine/Source/Programs/AutomationTool/Scripts/CopyBuildToStagingDirectory.Automation.cs
+++ b/Engine/Source/Programs/AutomationTool/Scripts/CopyBuildToStagingDirectory.Automation.cs
@@ -747,7 +747,7 @@ public partial class Project : CommandUtils
 				string VersionString = Params.ChunkInstallVersionString;
 				string ChunkInstallBasePath = CombinePaths(Params.ChunkInstallDirectory, SC.FinalCookPlatform);
 				string RawDataPath = CombinePaths(ChunkInstallBasePath, VersionString, PakName);
-				string RawDataPakPath = CombinePaths(RawDataPath, PakName + "-" + SC.FinalCookPlatform + ".pak");
+				string RawDataPakPath = CombinePaths(RawDataPath, PakName + "-" + SC.FinalCookPlatform + PostFix + ".pak");
 				//copy the pak chunk to the raw data folder
 				if (InternalUtils.SafeFileExists(RawDataPakPath, true))
 				{
@@ -759,6 +759,17 @@ public partial class Project : CommandUtils
 				{
 					InternalUtils.SafeDeleteFile(OutputLocation, true);
 				}
+				if (Params.IsGeneratingPatch)
+				{
+					if (String.IsNullOrEmpty(PatchSourceContentPath))
+					{
+						throw new AutomationException(String.Format("Failed Creating Chunk Install Data. No source pak for patch pak {0} given", OutputLocation));
+					}
+					// If we are generating a patch, then we need to copy the original pak across
+					// for distribution.
+					string SourceRawDataPakPath = CombinePaths(RawDataPath, PakName + "-" + SC.FinalCookPlatform + ".pak");
+					InternalUtils.SafeCopyFile(PatchSourceContentPath, SourceRawDataPakPath);
+				}
 
 				string BuildRoot = MakePathSafeToUseWithCommandLine(RawDataPath);
 				string CloudDir = MakePathSafeToUseWithCommandLine(CombinePaths(ChunkInstallBasePath, "CloudDir"));
diff --git a/Engine/Source/Runtime/HTTPChunkInstaller/Private/ChunkInstall.h b/Engine/Source/Runtime/HTTPChunkInstaller/Private/ChunkInstall.h
--- a/Engine/Source/Runtime/HTTPChunkInstaller/Private/ChunkInstall.h
+++ b/Engine/Source/Runtime/HTTPChunkInstaller/Private/ChunkInstall.h
@@ -98,6 +98,12 @@ public:
 		{
 			if (!CurrentMountPaks->Contains(PakFiles[PakIndex]) && !MountedPaks.Contains(PakFiles[PakIndex]))
 			{
+				// TODO: are we a patch?
+// 				if (PakFiles[PakIndex].EndsWith("_P.pak"))
+// 				{
+// 					// bump the read prioritiy 
+// 					++PakReadOrder;
+// 				}
 				if (FCoreDelegates::OnMountPak.IsBound())
 				{
 					auto bSuccess = FCoreDelegates::OnMountPak.Execute(PakFiles[PakIndex], PakReadOrder);
