- Fixed localization warning

diff --git a/Engine/Source/Developer/ProjectLauncher/Private/Widgets/Cook/SProjectLauncherCookByTheBookSettings.cpp b/Engine/Source/Developer/ProjectLauncher/Private/Widgets/Cook/SProjectLauncherCookByTheBookSettings.cpp
--- a/Engine/Source/Developer/ProjectLauncher/Private/Widgets/Cook/SProjectLauncherCookByTheBookSettings.cpp
+++ b/Engine/Source/Developer/ProjectLauncher/Private/Widgets/Cook/SProjectLauncherCookByTheBookSettings.cpp
@@ -403,7 +403,7 @@ TSharedRef<SWidget> SProjectLauncherCookByTheBookSettings::MakeComplexWidget()
 							[
 								// repository path text box
 								SNew(SEditableTextBox)
-								.ToolTipText(LOCTEXT("BasedOnReleaseVersionTextBoxTooltip", "Release version to base the next release / DLC / patch on."))
+								.ToolTipText(LOCTEXT("NextReleaseVersionTextBoxTooltip", "Release version to base the next release / DLC / patch on."))
 								.Text(this, &SProjectLauncherCookByTheBookSettings::HandleBasedOnReleaseVersionNameTextBlockText)
 								.OnTextCommitted(this, &SProjectLauncherCookByTheBookSettings::HandleBasedOnReleaseVersionNameCommitted)
 							]
