Copying //UE4/Dev-Sequencer to //UE4/Main (Source: //UE4/Dev-Sequencer @ 3617642)

diff --git a/Engine/Source/Developer/DeviceManager/Private/Widgets/SDeviceManager.cpp b/Engine/Source/Developer/DeviceManager/Private/Widgets/SDeviceManager.cpp
--- a/Engine/Source/Developer/DeviceManager/Private/Widgets/SDeviceManager.cpp
+++ b/Engine/Source/Developer/DeviceManager/Private/Widgets/SDeviceManager.cpp
@@ -38,7 +38,7 @@ static const FName DeviceToolbarTabId("DeviceToolbar");
 /* SDeviceManager constructors
  *****************************************************************************/
 
-SDeviceManager::SDeviceManager( )
+SDeviceManager::SDeviceManager()
 	: Model(MakeShareable(new FDeviceManagerModel()))
 { }
 
@@ -46,7 +46,7 @@ SDeviceManager::SDeviceManager( )
 /* SDeviceManager interface
  *****************************************************************************/
 
-void SDeviceManager::Construct( const FArguments& InArgs, const ITargetDeviceServiceManagerRef& InDeviceServiceManager, const TSharedRef<SDockTab>& ConstructUnderMajorTab, const TSharedPtr<SWindow>& ConstructUnderWindow )
+void SDeviceManager::Construct(const FArguments& InArgs, const TSharedRef<ITargetDeviceServiceManager>& InDeviceServiceManager, const TSharedRef<SDockTab>& ConstructUnderMajorTab, const TSharedPtr<SWindow>& ConstructUnderWindow)
 {
 	DeviceServiceManager = InDeviceServiceManager;
 
@@ -134,6 +134,7 @@ void SDeviceManager::Construct( const FArguments& InArgs, const ITargetDeviceSer
 		"Window"
 	);
 
+	// construct children
 	ChildSlot
 	[
 		SNew(SVerticalBox)
@@ -156,7 +157,7 @@ void SDeviceManager::Construct( const FArguments& InArgs, const ITargetDeviceSer
 /* SDeviceManager implementation
  *****************************************************************************/
 
-void SDeviceManager::BindCommands( )
+void SDeviceManager::BindCommands()
 {
 	const FDeviceDetailsCommands& Commands = FDeviceDetailsCommands::Get();
 
@@ -183,14 +184,14 @@ void SDeviceManager::BindCommands( )
 		FIsActionChecked::CreateSP(this, &SDeviceManager::HandleShareActionIsChecked));
 
 	// connectivity commands
-	UICommandList->MapAction( 
+	UICommandList->MapAction(
 		Commands.Connect, 
-		FExecuteAction::CreateSP( this, &SDeviceManager::HandleConnectActionExecute),
+		FExecuteAction::CreateSP(this, &SDeviceManager::HandleConnectActionExecute),
 		FCanExecuteAction::CreateSP(this, &SDeviceManager::HandleConnectActionCanExecute));
 
-	UICommandList->MapAction( 
+	UICommandList->MapAction(
 		Commands.Disconnect, 
-		FExecuteAction::CreateSP( this, &SDeviceManager::HandleDisconnectActionExecute),
+		FExecuteAction::CreateSP(this, &SDeviceManager::HandleDisconnectActionExecute),
 		FCanExecuteAction::CreateSP(this, &SDeviceManager::HandleDisconnectActionCanExecute));
 
 	// remote control commands
@@ -216,7 +217,7 @@ void SDeviceManager::BindCommands( )
 }
 
 
-void SDeviceManager::FillWindowMenu( FMenuBuilder& MenuBuilder, const TSharedPtr<FTabManager> TabManager )
+void SDeviceManager::FillWindowMenu(FMenuBuilder& MenuBuilder, const TSharedPtr<FTabManager> TabManager)
 {
 	if (!TabManager.IsValid())
 	{
@@ -231,7 +232,7 @@ void SDeviceManager::FillWindowMenu( FMenuBuilder& MenuBuilder, const TSharedPtr
 }
 
 
-bool SDeviceManager::ValidateDeviceAction( const ITargetDeviceRef& Device ) const
+bool SDeviceManager::ValidateDeviceAction(const ITargetDeviceRef& Device) const
 {
 	// @todo gmp: this needs to be improved, i.e. TargetPlatformManager::GetLocalDevice
 	if (Device->GetName() != FPlatformProcess::ComputerName())
@@ -248,7 +249,7 @@ bool SDeviceManager::ValidateDeviceAction( const ITargetDeviceRef& Device ) cons
 /* SDeviceManager callbacks
  *****************************************************************************/
 
-bool SDeviceManager::HandleClaimActionCanExecute( )
+bool SDeviceManager::HandleClaimActionCanExecute()
 {
 	ITargetDeviceServicePtr DeviceService = Model->GetSelectedDeviceService();
 
@@ -256,7 +257,7 @@ bool SDeviceManager::HandleClaimActionCanExecute( )
 }
 
 
-void SDeviceManager::HandleClaimActionExecute( )
+void SDeviceManager::HandleClaimActionExecute()
 {
 	ITargetDeviceServicePtr DeviceService = Model->GetSelectedDeviceService();
 
@@ -267,7 +268,7 @@ void SDeviceManager::HandleClaimActionExecute( )
 }
 
 
-bool SDeviceManager::HandleConnectActionCanExecute( )
+bool SDeviceManager::HandleConnectActionCanExecute()
 {
 	ITargetDeviceServicePtr DeviceService = Model->GetSelectedDeviceService();
 
@@ -285,7 +286,7 @@ bool SDeviceManager::HandleConnectActionCanExecute( )
 }
 
 
-void SDeviceManager::HandleConnectActionExecute( )
+void SDeviceManager::HandleConnectActionExecute()
 {
 	ITargetDeviceServicePtr DeviceService = Model->GetSelectedDeviceService();
 
@@ -304,7 +305,7 @@ void SDeviceManager::HandleConnectActionExecute( )
 }
 
 
-bool SDeviceManager::HandleDisconnectActionCanExecute( )
+bool SDeviceManager::HandleDisconnectActionCanExecute()
 {
 	ITargetDeviceServicePtr DeviceService = Model->GetSelectedDeviceService();
 
@@ -322,7 +323,7 @@ bool SDeviceManager::HandleDisconnectActionCanExecute( )
 }
 
 
-void SDeviceManager::HandleDisconnectActionExecute( )
+void SDeviceManager::HandleDisconnectActionExecute()
 {
 	ITargetDeviceServicePtr DeviceService = Model->GetSelectedDeviceService();
 
@@ -338,7 +339,7 @@ void SDeviceManager::HandleDisconnectActionExecute( )
 }
 
 
-bool SDeviceManager::HandlePowerOffActionCanExecute( )
+bool SDeviceManager::HandlePowerOffActionCanExecute()
 {
 	ITargetDeviceServicePtr DeviceService = Model->GetSelectedDeviceService();
 
@@ -356,7 +357,7 @@ bool SDeviceManager::HandlePowerOffActionCanExecute( )
 }
 
 
-void SDeviceManager::HandlePowerOffActionExecute( bool Force )
+void SDeviceManager::HandlePowerOffActionExecute(bool Force)
 {
 	ITargetDeviceServicePtr DeviceService = Model->GetSelectedDeviceService();
 
@@ -372,7 +373,7 @@ void SDeviceManager::HandlePowerOffActionExecute( bool Force )
 }
 
 
-bool SDeviceManager::HandlePowerOnActionCanExecute( )
+bool SDeviceManager::HandlePowerOnActionCanExecute()
 {
 	ITargetDeviceServicePtr DeviceService = Model->GetSelectedDeviceService();
 
@@ -390,7 +391,7 @@ bool SDeviceManager::HandlePowerOnActionCanExecute( )
 }
 
 
-void SDeviceManager::HandlePowerOnActionExecute( )
+void SDeviceManager::HandlePowerOnActionExecute()
 {
 	ITargetDeviceServicePtr DeviceService = Model->GetSelectedDeviceService();
 
@@ -406,7 +407,7 @@ void SDeviceManager::HandlePowerOnActionExecute( )
 }
 
 
-bool SDeviceManager::HandleRebootActionCanExecute( )
+bool SDeviceManager::HandleRebootActionCanExecute()
 {
 	ITargetDeviceServicePtr DeviceService = Model->GetSelectedDeviceService();
 
@@ -424,7 +425,7 @@ bool SDeviceManager::HandleRebootActionCanExecute( )
 }
 
 
-void SDeviceManager::HandleRebootActionExecute( )
+void SDeviceManager::HandleRebootActionExecute()
 {
 	ITargetDeviceServicePtr DeviceService = Model->GetSelectedDeviceService();
 
@@ -440,7 +441,7 @@ void SDeviceManager::HandleRebootActionExecute( )
 }
 
 
-bool SDeviceManager::HandleReleaseActionCanExecute( )
+bool SDeviceManager::HandleReleaseActionCanExecute()
 {
 	ITargetDeviceServicePtr DeviceService = Model->GetSelectedDeviceService();
 
@@ -448,7 +449,7 @@ bool SDeviceManager::HandleReleaseActionCanExecute( )
 }
 
 
-void SDeviceManager::HandleReleaseActionExecute( )
+void SDeviceManager::HandleReleaseActionExecute()
 {
 	ITargetDeviceServicePtr DeviceService = Model->GetSelectedDeviceService();
 
@@ -459,7 +460,7 @@ void SDeviceManager::HandleReleaseActionExecute( )
 }
 
 
-bool SDeviceManager::HandleRemoveActionCanExecute( )
+bool SDeviceManager::HandleRemoveActionCanExecute()
 {
 	ITargetDeviceServicePtr DeviceService = Model->GetSelectedDeviceService();
 
@@ -468,7 +469,7 @@ bool SDeviceManager::HandleRemoveActionCanExecute( )
 }
 
 
-void SDeviceManager::HandleRemoveActionExecute( )
+void SDeviceManager::HandleRemoveActionExecute()
 {
 	ITargetDeviceServicePtr DeviceService = Model->GetSelectedDeviceService();
 
@@ -479,7 +480,7 @@ void SDeviceManager::HandleRemoveActionExecute( )
 }
 
 
-bool SDeviceManager::HandleShareActionIsChecked( )
+bool SDeviceManager::HandleShareActionIsChecked()
 {
 	ITargetDeviceServicePtr DeviceService = Model->GetSelectedDeviceService();
 
@@ -487,7 +488,7 @@ bool SDeviceManager::HandleShareActionIsChecked( )
 }
 
 
-void SDeviceManager::HandleShareActionExecute( )
+void SDeviceManager::HandleShareActionExecute()
 {
 	ITargetDeviceServicePtr DeviceService = Model->GetSelectedDeviceService();
 
@@ -498,7 +499,7 @@ void SDeviceManager::HandleShareActionExecute( )
 }
 
 
-bool SDeviceManager::HandleShareActionCanExecute( )
+bool SDeviceManager::HandleShareActionCanExecute()
 {
 	ITargetDeviceServicePtr DeviceService = Model->GetSelectedDeviceService();
 
@@ -506,7 +507,7 @@ bool SDeviceManager::HandleShareActionCanExecute( )
 }
 
 
-TSharedRef<SDockTab> SDeviceManager::HandleTabManagerSpawnTab( const FSpawnTabArgs& Args, FName TabIdentifier )
+TSharedRef<SDockTab> SDeviceManager::HandleTabManagerSpawnTab(const FSpawnTabArgs& Args, FName TabIdentifier)
 {
 	TSharedPtr<SWidget> TabWidget = SNullWidget::NullWidget;
 	bool AutoSizeTab = false;
diff --git a/Engine/Source/Developer/DeviceManager/Private/Widgets/SDeviceManager.h b/Engine/Source/Developer/DeviceManager/Private/Widgets/SDeviceManager.h
--- a/Engine/Source/Developer/DeviceManager/Private/Widgets/SDeviceManager.h
+++ b/Engine/Source/Developer/DeviceManager/Private/Widgets/SDeviceManager.h
@@ -6,7 +6,7 @@
 #include "SlateFwd.h"
 #include "Widgets/DeclarativeSyntaxSupport.h"
 #include "Widgets/SCompoundWidget.h"
-#include "Interfaces/ITargetDeviceServiceManager.h"
+#include "ITargetDeviceServiceManager.h"
 #include "Models/DeviceManagerModel.h"
 
 class FMenuBuilder;
@@ -15,6 +15,7 @@ class FTabManager;
 class FUICommandList;
 class SWindow;
 
+
 /**
  * Implements the device manager front-end widget.
  */
@@ -28,119 +29,117 @@ public:
 
 public:
 
-	/**
-	 * Default constructor.
-	 */
-	SDeviceManager( );
+	/** Default constructor. */
+	SDeviceManager();
 
 public:
 
 	/**
-	 * Constructs the widget.
+	 * Construct the widget.
 	 *
 	 * @param InArgs The Slate argument list.
 	 * @param InDeviceServiceManager The target device manager to use.
 	 * @param ConstructUnderMajorTab The major tab which will contain the session front-end.
 	 * @param ConstructUnderWindow The window in which this widget is being constructed.
 	 */
-	void Construct( const FArguments& InArgs, const ITargetDeviceServiceManagerRef& InDeviceServiceManager, const TSharedRef<SDockTab>& ConstructUnderMajorTab, const TSharedPtr<SWindow>& ConstructUnderWindow );
+	void Construct(const FArguments& InArgs, const TSharedRef<ITargetDeviceServiceManager>& InDeviceServiceManager, const TSharedRef<SDockTab>& ConstructUnderMajorTab, const TSharedPtr<SWindow>& ConstructUnderWindow);
 
 protected:
 
 	/**
-	 * Binds the device commands on our toolbar.
+	 * Bind the device commands on our toolbar.
 	 */
-	void BindCommands( );
+	void BindCommands();
 
 	/**
-	* Fills the Window menu with menu items.
-	*
-	* @param MenuBuilder The multi-box builder that should be filled with content for this pull-down menu.
-	* @param TabManager A Tab Manager from which to populate tab spawner menu items.
-	*/
+	 * Fill the Window menu with menu items.
+	 *
+	 * @param MenuBuilder The multi-box builder that should be filled with content for this pull-down menu.
+	 * @param TabManager A Tab Manager from which to populate tab spawner menu items.
+	 */
 	static void FillWindowMenu(FMenuBuilder& MenuBuilder, const TSharedPtr<FTabManager> TabManager);
 
 	/**
-	 * Validates actions on the specified device.
+	 * Validate actions on the specified device.
 	 *
 	 * @param Device The device to perform an action on.
 	 * @return true if actions on the device are permitted, false otherwise.
 	 */
-	bool ValidateDeviceAction( const ITargetDeviceRef& Device ) const;
+	bool ValidateDeviceAction(const ITargetDeviceRef& Device) const;
 
 private:
 
-	// Callback for determining whether the 'Claim' action can execute.
-	bool HandleClaimActionCanExecute( );
+	/** Callback for determining whether the 'Claim' action can execute. */
+	bool HandleClaimActionCanExecute();
 
-	// Callback for executing the 'Claim' action.
-	void HandleClaimActionExecute( );
+	/** Callback for executing the 'Claim' action. */
+	void HandleClaimActionExecute();
 
-	// Callback for determining whether the "Connect" action can execute
-	bool HandleConnectActionCanExecute( );
+	/** Callback for determining whether the "Connect" action can execute. */
+	bool HandleConnectActionCanExecute();
 
-	// Callback for executing the "Connect" action
-	void HandleConnectActionExecute( );
+	/** Callback for executing the "Connect" action. */
+	void HandleConnectActionExecute();
 
-	// Callback for determining whether the "Disconnect" action can execute
-	bool HandleDisconnectActionCanExecute( );
+	/** Callback for determining whether the "Disconnect" action can execute. */
+	bool HandleDisconnectActionCanExecute();
 
-	// Callback for executing the "Disconnect" action
-	void HandleDisconnectActionExecute( );
+	/** Callback for executing the "Disconnect" action. */
+	void HandleDisconnectActionExecute();
 
-	// Callback for determining whether the 'Power Off' action can execute.
-	bool HandlePowerOffActionCanExecute( );
+	/** Callback for determining whether the 'Power Off' action can execute. */
+	bool HandlePowerOffActionCanExecute();
 
-	// Callback for executing the 'Power Off' action.
-	void HandlePowerOffActionExecute( bool Force );
+	/** Callback for executing the 'Power Off' action. */
+	void HandlePowerOffActionExecute(bool Force);
 
-	// Callback for determining whether the 'Power On' action can execute.
-	bool HandlePowerOnActionCanExecute( );
+	/** Callback for determining whether the 'Power On' action can execute. */
+	bool HandlePowerOnActionCanExecute();
 
-	// Callback for executing the 'Power On' action.
-	void HandlePowerOnActionExecute( );
+	/** Callback for executing the 'Power On' action. */
+	void HandlePowerOnActionExecute();
 
-	// Callback for determining whether the 'Reboot' action can execute.
-	bool HandleRebootActionCanExecute( );
+	/** Callback for determining whether the 'Reboot' action can execute. */
+	bool HandleRebootActionCanExecute();
 
-	// Callback for executing the 'Reboot' action.
-	void HandleRebootActionExecute( );
+	/** Callback for executing the 'Reboot' action. */
+	void HandleRebootActionExecute();
 
-	// Callback for determining whether the 'Release' action can execute.
-	bool HandleReleaseActionCanExecute( );
+	/** Callback for determining whether the 'Release' action can execute. */
+	bool HandleReleaseActionCanExecute();
 
-	// Callback for executing the 'Release' action.
-	void HandleReleaseActionExecute( );
+	/** Callback for executing the 'Release' action. */
+	void HandleReleaseActionExecute();
 
-	// Callback for determining whether the 'Remove' action can execute.
-	bool HandleRemoveActionCanExecute( );
+	/** Callback for determining whether the 'Remove' action can execute. */
+	bool HandleRemoveActionCanExecute();
 
-	// Callback for executing the 'Remove' action.
-	void HandleRemoveActionExecute( );
+	/** Callback for executing the 'Remove' action. */
+	void HandleRemoveActionExecute();
 
-	// Callback for determining whether the 'Share' action can execute.
-	bool HandleShareActionCanExecute( );
+	/** Callback for determining whether the 'Share' action can execute. */
+	bool HandleShareActionCanExecute();
 
-	// Callback for executing the 'Share' action.
-	void HandleShareActionExecute( );
+	/** Callback for executing the 'Share' action. */
+	void HandleShareActionExecute();
 
-	// Callback for determining the checked state of the 'Share' action.
-	bool HandleShareActionIsChecked( );
+	/** Callback for determining the checked state of the 'Share' action. */
+	bool HandleShareActionIsChecked();
 
-	// Callback for spawning tabs.
-	TSharedRef<SDockTab> HandleTabManagerSpawnTab( const FSpawnTabArgs& Args, FName TabIdentifier );
+	/** Callback for spawning tabs. */
+	TSharedRef<SDockTab> HandleTabManagerSpawnTab(const FSpawnTabArgs& Args, FName TabIdentifier);
 
 private:
 
-	// Holds the target device service manager.
-	ITargetDeviceServiceManagerPtr DeviceServiceManager;
+	/** Holds the target device service manager. */
+	TSharedPtr<ITargetDeviceServiceManager> DeviceServiceManager;
 
-	// Holds the device manager's view model.
-	FDeviceManagerModelRef Model;
+	/** Holds the device manager's view model. */
+	TSharedRef<FDeviceManagerModel> Model;
 
-	// Holds the tab manager that manages the front-end's tabs.
+	/** Holds the tab manager that manages the front-end's tabs. */
 	TSharedPtr<FTabManager> TabManager;
 
-	// The command list for controlling the device
+	/** The command list for controlling the device. */
 	TSharedPtr<FUICommandList> UICommandList;
 };
diff --git a/Engine/Source/Developer/DeviceManager/Private/Widgets/Shared/SDeviceQuickInfo.h b/Engine/Source/Developer/DeviceManager/Private/Widgets/Shared/SDeviceQuickInfo.h
--- a/Engine/Source/Developer/DeviceManager/Private/Widgets/Shared/SDeviceQuickInfo.h
+++ b/Engine/Source/Developer/DeviceManager/Private/Widgets/Shared/SDeviceQuickInfo.h
@@ -3,22 +3,24 @@
 #pragma once
 
 #include "CoreMinimal.h"
-#include "Widgets/DeclarativeSyntaxSupport.h"
-#include "Widgets/SCompoundWidget.h"
-#include "Interfaces/ITargetDeviceService.h"
+#include "EditorStyleSet.h"
 #include "Fonts/SlateFontInfo.h"
+#include "ITargetDeviceService.h"
 #include "Misc/Paths.h"
+#include "PlatformInfo.h"
 #include "Styling/StyleDefaults.h"
+#include "Widgets/DeclarativeSyntaxSupport.h"
 #include "Widgets/SBoxPanel.h"
-#include "Widgets/Layout/SBox.h"
-#include "Widgets/Text/STextBlock.h"
-#include "EditorStyleSet.h"
+#include "Widgets/SCompoundWidget.h"
 #include "Widgets/Images/SImage.h"
+#include "Widgets/Layout/SBox.h"
 #include "Widgets/Layout/SGridPanel.h"
-#include "PlatformInfo.h"
+#include "Widgets/Text/STextBlock.h"
+
 
 #define LOCTEXT_NAMESPACE "SDeviceQuickInfo"
 
+
 /**
  * Implements a tool tip for widget the device browser.
  */
@@ -37,13 +39,13 @@ public:
 public:
 
 	/**
-	 * Constructs the widget.
+	 * Construct the widget.
 	 *
 	 * @param InArgs The construction arguments.
 	 * @param InDeviceServiceManager The target device service manager to use.
 	 * @param InDeviceManagerState The optional device manager view state.
 	 */
-	void Construct( const FArguments& InArgs )
+	void Construct(const FArguments& InArgs)
 	{
 		DeviceService = InArgs._InitialDeviceService.Get();
 
@@ -172,19 +174,19 @@ public:
 public:
 
 	/**
-	 * Sets the device service whose information is being shown.
+	 * Set the device service whose information is being shown.
 	 *
 	 * @param InDeviceService The device service to show.
 	 */
-	void SetDeviceService( const ITargetDeviceServicePtr& InDeviceService )
+	void SetDeviceService(const TSharedPtr<ITargetDeviceService, ESPMode::ThreadSafe>& InDeviceService)
 	{
 		DeviceService = InDeviceService;
 	}
 
 private:
 
-	// Callback for getting the device's unique identifier.
-	FText HandleDeviceIdText( ) const
+	/** Callback for getting the device's unique identifier. */
+	FText HandleDeviceIdText() const
 	{
 		if (DeviceService.IsValid())
 		{
@@ -198,8 +200,8 @@ private:
 		return LOCTEXT("UnknownValue", "<unknown>");
 	}
 
-	// Callback for getting the name of the shown device.
-	FText HandleDeviceNameText( ) const
+	/** Callback for getting the name of the shown device. */
+	FText HandleDeviceNameText() const
 	{
 		if (DeviceService.IsValid())
 		{
@@ -214,8 +216,8 @@ private:
 		return LOCTEXT("UnknownValue", "<unknown>");
 	}
 
-	// Callback for getting the text that indicates whether the shown device is the platform's default device.
-	FText HandleIsDefaultText( ) const
+	/** Callback for getting the text that indicates whether the shown device is the platform's default device. */
+	FText HandleIsDefaultText() const
 	{
 		if (DeviceService.IsValid())
 		{
@@ -232,8 +234,8 @@ private:
 		return LOCTEXT("UnknownValue", "<unknown>");
 	}
 
-	// Callback for getting the make and model of the shown device.
-	FText HandleOperatingSystemText( ) const
+	/** Callback for getting the make and model of the shown device. */
+	FText HandleOperatingSystemText() const
 	{
 		if (DeviceService.IsValid())
 		{
@@ -253,8 +255,8 @@ private:
 		return LOCTEXT("UnknownValue", "<unknown>");
 	}
 
-	// Callback for getting the icon of the device's platform.
-	const FSlateBrush* HandlePlatformIcon( ) const
+	/** Callback for getting the icon of the device's platform. */
+	const FSlateBrush* HandlePlatformIcon() const
 	{
 		if (DeviceService.IsValid())
 		{
@@ -268,8 +270,8 @@ private:
 		return FStyleDefaults::GetNoBrush();
 	}
 
-	// Callback for getting the name of the device's platform.
-	FText HandlePlatformNameText( ) const
+	/** Callback for getting the name of the device's platform. */
+	FText HandlePlatformNameText() const
 	{
 		if (DeviceService.IsValid())
 		{
@@ -284,8 +286,8 @@ private:
 		return LOCTEXT("UnknownValue", "<unknown>");
 	}
 
-	// Callback for getting the status of the device.
-	FText HandleStatusText( ) const
+	/** Callback for getting the status of the device. */
+	FText HandleStatusText() const
 	{
 		if (DeviceService.IsValid())
 		{
@@ -309,8 +311,8 @@ private:
 
 private:
 
-	// Holds the service for the device whose details are being shown.
-	ITargetDeviceServicePtr DeviceService;
+	/** The service for the device whose details are being shown. */
+	TSharedPtr<ITargetDeviceService, ESPMode::ThreadSafe> DeviceService;
 };
 
 
diff --git a/Engine/Source/Developer/DeviceManager/Private/Widgets/Toolbar/SDeviceToolbar.cpp b/Engine/Source/Developer/DeviceManager/Private/Widgets/Toolbar/SDeviceToolbar.cpp
--- a/Engine/Source/Developer/DeviceManager/Private/Widgets/Toolbar/SDeviceToolbar.cpp
+++ b/Engine/Source/Developer/DeviceManager/Private/Widgets/Toolbar/SDeviceToolbar.cpp
@@ -1,12 +1,15 @@
 // Copyright 1998-2017 Epic Games, Inc. All Rights Reserved.
 
-#include "Widgets/Toolbar/SDeviceToolbar.h"
-#include "SlateOptMacros.h"
-#include "Widgets/Layout/SBorder.h"
+#include "SDeviceToolbar.h"
+
+#include "EditorStyleSet.h"
 #include "Framework/MultiBox/MultiBoxDefs.h"
 #include "Framework/MultiBox/MultiBoxBuilder.h"
-#include "EditorStyleSet.h"
+#include "SlateOptMacros.h"
+#include "Widgets/Layout/SBorder.h"
+
 #include "Models/DeviceDetailsCommands.h"
+#include "Models/DeviceManagerModel.h"
 
 
 #define LOCTEXT_NAMESPACE "SDeviceToolbar"
@@ -16,10 +19,15 @@
  *****************************************************************************/
 
 BEGIN_SLATE_FUNCTION_BUILD_OPTIMIZATION
-void SDeviceToolbar::Construct( const FArguments& InArgs, const FDeviceManagerModelRef& InModel, const TSharedPtr<FUICommandList>& InUICommandList )
+void SDeviceToolbar::Construct(const FArguments& InArgs, const TSharedRef<FDeviceManagerModel>& InModel, const TSharedPtr<FUICommandList>& InUICommandList)
 {
 	Model = InModel;
 
+	// callback for getting the enabled state of the toolbar.
+	auto ToolbarIsEnabled = [this]() -> bool {
+		return Model->GetSelectedDeviceService().IsValid();
+	};
+
 	// create the toolbar
 	FToolBarBuilder Toolbar(InUICommandList, FMultiBoxCustomization::None);
 	{
@@ -38,11 +46,12 @@ void SDeviceToolbar::Construct( const FArguments& InArgs, const FDeviceManagerMo
 		Toolbar.AddToolBarButton(FDeviceDetailsCommands::Get().Reboot);
 	}
 
+	// construct children
 	ChildSlot
 	[
 		SNew(SBorder)
 			.BorderImage(FEditorStyle::GetBrush("ToolPanel.GroupBorder"))
-			.IsEnabled(this, &SDeviceToolbar::HandleToolbarIsEnabled)
+			.IsEnabled_Lambda(ToolbarIsEnabled)
 			.Padding(0.0f)
 			[
 				Toolbar.MakeWidget()
@@ -51,13 +60,5 @@ void SDeviceToolbar::Construct( const FArguments& InArgs, const FDeviceManagerMo
 }
 END_SLATE_FUNCTION_BUILD_OPTIMIZATION
 
-/* SDeviceToolbar callbacks
- *****************************************************************************/
-
-bool SDeviceToolbar::HandleToolbarIsEnabled( ) const
-{
-	return Model->GetSelectedDeviceService().IsValid();
-}
-
 
 #undef LOCTEXT_NAMESPACE
diff --git a/Engine/Source/Developer/DeviceManager/Private/Widgets/Toolbar/SDeviceToolbar.h b/Engine/Source/Developer/DeviceManager/Private/Widgets/Toolbar/SDeviceToolbar.h
--- a/Engine/Source/Developer/DeviceManager/Private/Widgets/Toolbar/SDeviceToolbar.h
+++ b/Engine/Source/Developer/DeviceManager/Private/Widgets/Toolbar/SDeviceToolbar.h
@@ -2,13 +2,13 @@
 
 #pragma once
 
-#include "CoreMinimal.h"
 #include "Widgets/DeclarativeSyntaxSupport.h"
 #include "Widgets/SCompoundWidget.h"
-#include "Models/DeviceManagerModel.h"
 
+class FDeviceManagerModel;
 class FUICommandList;
 
+
 /**
  * Implements the device toolbar widget.
  */
@@ -29,15 +29,10 @@ public:
 	 * @param InModel The view model to use.
 	 * @param InUICommandList The UI command list to use.
 	 */
-	void Construct( const FArguments& InArgs, const FDeviceManagerModelRef& InModel, const TSharedPtr<FUICommandList>& InUICommandList );
-
-private:
-
-	// Callback for getting the enabled state of the toolbar.
-	bool HandleToolbarIsEnabled( ) const;
+	void Construct(const FArguments& InArgs, const TSharedRef<FDeviceManagerModel>& InModel, const TSharedPtr<FUICommandList>& InUICommandList);
 
 private:
 
-	// Holds a pointer the device manager's view model.
-	FDeviceManagerModelPtr Model;
+	/** Pointer the device manager's view model. */
+	TSharedPtr<FDeviceManagerModel> Model;
 };
diff --git a/Engine/Source/Developer/ExternalImagePicker/Private/SExternalImagePicker.cpp b/Engine/Source/Developer/ExternalImagePicker/Private/SExternalImagePicker.cpp
--- a/Engine/Source/Developer/ExternalImagePicker/Private/SExternalImagePicker.cpp
+++ b/Engine/Source/Developer/ExternalImagePicker/Private/SExternalImagePicker.cpp
@@ -9,7 +9,8 @@
 #include "Widgets/Input/SButton.h"
 #include "SResetToDefaultMenu.h"
 #include "DesktopPlatformModule.h"
-#include "Interfaces/IImageWrapperModule.h"
+#include "IImageWrapper.h"
+#include "IImageWrapperModule.h"
 #include "Widgets/Layout/SEnableBox.h"
 
 #define LOCTEXT_NAMESPACE "ExternalImagePicker"
@@ -255,7 +256,7 @@ TSharedPtr< FSlateDynamicImageBrush > SExternalImagePicker::LoadImageAsBrush( co
 	if( FFileHelper::LoadFileToArray( RawFileData, *ImagePath ) )
 	{
 		IImageWrapperModule& ImageWrapperModule = FModuleManager::LoadModuleChecked<IImageWrapperModule>( FName("ImageWrapper") );
-		IImageWrapperPtr ImageWrappers[4] =
+		TSharedPtr<IImageWrapper> ImageWrappers[4] =
 		{ 
 			ImageWrapperModule.CreateImageWrapper( EImageFormat::PNG ),
 			ImageWrapperModule.CreateImageWrapper( EImageFormat::BMP ),
diff --git a/Engine/Source/Developer/HTML5/HTML5TargetPlatform/Private/HTML5TargetDevice.h b/Engine/Source/Developer/HTML5/HTML5TargetPlatform/Private/HTML5TargetDevice.h
--- a/Engine/Source/Developer/HTML5/HTML5TargetPlatform/Private/HTML5TargetDevice.h
+++ b/Engine/Source/Developer/HTML5/HTML5TargetPlatform/Private/HTML5TargetDevice.h
@@ -1,9 +1,5 @@
 // Copyright 1998-2017 Epic Games, Inc. All Rights Reserved.
 
-/*=============================================================================
-	HTML5TargetDevice.h: Declares the HTML5TargetDevice class.
-=============================================================================*/
-
 #pragma once
 
 #include "CoreMinimal.h"
@@ -13,6 +9,7 @@
 class FHTML5TargetDevice;
 class ITargetPlatform;
 
+
 /**
  * Type definition for shared pointers to instances of FHTML5TargetDevice.
  */
diff --git a/Engine/Source/Developer/IOS/IOSTargetPlatform/Private/IOSTargetDevice.cpp b/Engine/Source/Developer/IOS/IOSTargetPlatform/Private/IOSTargetDevice.cpp
--- a/Engine/Source/Developer/IOS/IOSTargetPlatform/Private/IOSTargetDevice.cpp
+++ b/Engine/Source/Developer/IOS/IOSTargetPlatform/Private/IOSTargetDevice.cpp
@@ -1,7 +1,28 @@
 // Copyright 1998-2017 Epic Games, Inc. All Rights Reserved.
 
 #include "IOSTargetDevice.h"
+
+#include "HAL/PlatformProcess.h"
 #include "IOSMessageProtocol.h"
+#include "Interfaces/ITargetPlatform.h"
+#include "MessageEndpoint.h"
+#include "MessageEndpointBuilder.h"
+
+
+FIOSTargetDevice::FIOSTargetDevice(const ITargetPlatform& InTargetPlatform)
+	: TargetPlatform(InTargetPlatform)
+	, DeviceEndpoint()
+	, AppId()
+	, bCanReboot(false)
+	, bCanPowerOn(false)
+	, bCanPowerOff(false)
+	, DeviceType(ETargetDeviceTypes::Indeterminate)
+{
+	DeviceId = FTargetDeviceId(TargetPlatform.PlatformName(), FPlatformProcess::ComputerName());
+	DeviceName = FPlatformProcess::ComputerName();
+	MessageEndpoint = FMessageEndpoint::Builder("FIOSTargetDevice").Build();
+}
+
 
 bool FIOSTargetDevice::Connect()
 {
diff --git a/Engine/Source/Developer/IOS/IOSTargetPlatform/Private/IOSTargetDevice.h b/Engine/Source/Developer/IOS/IOSTargetPlatform/Private/IOSTargetDevice.h
--- a/Engine/Source/Developer/IOS/IOSTargetPlatform/Private/IOSTargetDevice.h
+++ b/Engine/Source/Developer/IOS/IOSTargetPlatform/Private/IOSTargetDevice.h
@@ -1,52 +1,41 @@
 // Copyright 1998-2017 Epic Games, Inc. All Rights Reserved.
 
-/*=============================================================================
-	IOSTargetDevice.h: Defines an IOS device for communication
-=============================================================================*/
-
 #pragma once
 
 #include "CoreMinimal.h"
-#include "Interfaces/ITargetDevice.h"
-#include "Interfaces/ITargetPlatform.h"
-#include "HAL/PlatformProcess.h"
 #include "IMessageContext.h"
-#include "Helpers/MessageEndpoint.h"
-#include "Helpers/MessageEndpointBuilder.h"
+#include "Interfaces/ITargetDevice.h"
 
+class FMessageEndpoint;
+class ITargetPlatform;
 
-/**
- * Type definition for shared pointers to instances of FIOSTargetDevice.
- */
+
+/** Type definition for shared pointers to instances of FIOSTargetDevice. */
 typedef TSharedPtr<class FIOSTargetDevice, ESPMode::ThreadSafe> FIOSTargetDevicePtr;
 
-/**
- * Type definition for shared references to instances of FIOSTargetDevice.
- */
+/** Type definition for shared references to instances of FIOSTargetDevice. */
 typedef TSharedRef<class FIOSTargetDevice, ESPMode::ThreadSafe> FIOSTargetDeviceRef;
 
 
 /**
  * Implements an iOS target device.
  */
-class FIOSTargetDevice : public ITargetDevice
+class FIOSTargetDevice
+	: public ITargetDevice
 {
 public:
-	FIOSTargetDevice (const ITargetPlatform& InTargetPlatform)
-		: TargetPlatform(InTargetPlatform)
-		, DeviceEndpoint()
-		, AppId()
-		, bCanReboot(false)
-		, bCanPowerOn(false)
-		, bCanPowerOff(false)
-		, DeviceType(ETargetDeviceTypes::Indeterminate)
-	{
-		DeviceId = FTargetDeviceId(TargetPlatform.PlatformName(), FPlatformProcess::ComputerName());
-		DeviceName = FPlatformProcess::ComputerName();
-		MessageEndpoint = FMessageEndpoint::Builder("FIOSTargetDevice").Build();
-	}
+
+	/**
+	 * Create and initialize a new instance.
+	 *
+	 * @param InTargetPlatform The target platform that owns the device.
+	 */
+	FIOSTargetDevice(const ITargetPlatform& InTargetPlatform);
+
 public:
-	// ITargetDevice interface
+
+	//~ ITargetDevice interface
+
 	virtual bool Connect() override;
 	virtual bool Deploy(const FString& SourceFolder, FString& OutAppId) override;
 	virtual void Disconnect() override;
@@ -68,9 +57,9 @@ public:
 	virtual bool TerminateProcess(const int64 ProcessId) override;
 	virtual void SetUserCredentials(const FString& UserName, const FString& UserPassword) override;
 	virtual bool GetUserCredentials(FString& OutUserName, FString& OutUserPassword) override;
-	// End of ITargetDevice interface
 
 public:
+
 	/** Timeout check for removing stale devices */
 	FDateTime LastPinged;
 
@@ -78,32 +67,42 @@ private:
 
 	/** The current status of this device. */
 //	ETargetDeviceStatus::Type Status;
-	/** Holds a reference to the device's target platform. */
+
+/** Holds a reference to the device's target platform. */
 	const ITargetPlatform& TargetPlatform;
+
 	/** Contains the address of the remote device */
 	FMessageAddress DeviceEndpoint;
+
 	/** MessageEndpoint for communicating with remote device */
-	FMessageEndpointPtr MessageEndpoint;
+	TSharedPtr<FMessageEndpoint, ESPMode::ThreadSafe> MessageEndpoint;
+
 	/** Contains the current AppID/GameName for Deployment/launching. */
 	FString AppId;
-	/** Containts the build configuration of the app to deploy */
+
+	/** Contains the build configuration of the app to deploy */
 	EBuildConfigurations::Type BuildConfiguration;
+
 	/** Lets us know whether the thing is a sim device or a physical device. */
 	bool bIsSimulated;
 
-
 private:
 
 	/** Remote rebootable */
 	bool bCanReboot;
+
 	/** Remote bootable */
 	bool bCanPowerOn;
+
 	/** Remote shutdown-able */
 	bool bCanPowerOff;
+
 	/** Id of device */
 	FTargetDeviceId DeviceId;
+
 	/** Name of device */
 	FString DeviceName;
+
 	/** Type of device */
 	ETargetDeviceTypes DeviceType;
 
@@ -111,15 +110,15 @@ public:
 
 	void SetFeature(ETargetDeviceFeatures InFeature, bool bFlag)
 	{
-		if(InFeature == ETargetDeviceFeatures::Reboot)
+		if (InFeature == ETargetDeviceFeatures::Reboot)
 		{
 			bCanReboot = bFlag;
 		}
-		else if(InFeature == ETargetDeviceFeatures::PowerOn)
+		else if (InFeature == ETargetDeviceFeatures::PowerOn)
 		{
 			bCanPowerOn = bFlag;
 		}
-		else if(InFeature == ETargetDeviceFeatures::PowerOff)
+		else if (InFeature == ETargetDeviceFeatures::PowerOff)
 		{
 			bCanPowerOff = bFlag;
 		}
@@ -140,19 +139,19 @@ public:
 	/** Sets the type of the device */
 	void SetDeviceType(const FString InDeviceTypeString)
 	{
-		if(InDeviceTypeString == TEXT("Browser"))
+		if (InDeviceTypeString == TEXT("Browser"))
 		{
 			DeviceType = ETargetDeviceTypes::Browser;
 		}
-		else if(InDeviceTypeString == TEXT("Console"))
+		else if (InDeviceTypeString == TEXT("Console"))
 		{
 			DeviceType = ETargetDeviceTypes::Console;
 		}
-		else if(InDeviceTypeString == TEXT("Phone"))
+		else if (InDeviceTypeString == TEXT("Phone"))
 		{
 			DeviceType = ETargetDeviceTypes::Phone;
 		}
-		else if(InDeviceTypeString == TEXT("Tablet"))
+		else if (InDeviceTypeString == TEXT("Tablet"))
 		{
 			DeviceType = ETargetDeviceTypes::Tablet;
 		}
diff --git a/Engine/Source/Developer/IOS/IOSTargetPlatform/Private/IOSTargetPlatform.cpp b/Engine/Source/Developer/IOS/IOSTargetPlatform/Private/IOSTargetPlatform.cpp
--- a/Engine/Source/Developer/IOS/IOSTargetPlatform/Private/IOSTargetPlatform.cpp
+++ b/Engine/Source/Developer/IOS/IOSTargetPlatform/Private/IOSTargetPlatform.cpp
@@ -274,7 +274,7 @@ void FIOSTargetPlatform::PingNetworkDevices()
 /* FIOSTargetPlatform callbacks
  *****************************************************************************/
 
-void FIOSTargetPlatform::HandlePongMessage( const FIOSLaunchDaemonPong& Message, const IMessageContextRef& Context )
+void FIOSTargetPlatform::HandlePongMessage( const FIOSLaunchDaemonPong& Message, const TSharedRef<IMessageContext, ESPMode::ThreadSafe>& Context )
 {
 	FTargetDeviceId DeviceId;
 	FTargetDeviceId::Parse(Message.DeviceID, DeviceId);
@@ -330,7 +330,7 @@ void FIOSTargetPlatform::HandleDeviceConnected(const FIOSLaunchDaemonPong& Messa
 	}
 	
 	// Add a very long time period to prevent the devices from getting disconnected due to a lack of pong messages
-	Device->LastPinged = FDateTime::UtcNow() + FTimespan(100, 0, 0, 0, 0);
+	Device->LastPinged = FDateTime::UtcNow() + FTimespan::FromDays(100.0);
 }
 
 
@@ -348,7 +348,7 @@ void FIOSTargetPlatform::HandleDeviceDisconnected(const FIOSLaunchDaemonPong& Me
 	}
 }
 
-bool FIOSTargetPlatform::HandleTicker(float DeltaTime )
+bool FIOSTargetPlatform::HandleTicker(float DeltaTime)
 {
 	PingNetworkDevices();
 
diff --git a/Engine/Source/Developer/IOS/IOSTargetPlatform/Private/IOSTargetPlatform.h b/Engine/Source/Developer/IOS/IOSTargetPlatform/Private/IOSTargetPlatform.h
--- a/Engine/Source/Developer/IOS/IOSTargetPlatform/Private/IOSTargetPlatform.h
+++ b/Engine/Source/Developer/IOS/IOSTargetPlatform/Private/IOSTargetPlatform.h
@@ -173,7 +173,7 @@ private:
 	bool HandleTicker( float DeltaTime );
 
 	// Handles received pong messages from the LauncherDaemon.
-	void HandlePongMessage( const FIOSLaunchDaemonPong& Message, const IMessageContextRef& Context );
+	void HandlePongMessage( const FIOSLaunchDaemonPong& Message, const TSharedRef<IMessageContext, ESPMode::ThreadSafe>& Context );
 
     void HandleDeviceConnected( const FIOSLaunchDaemonPong& Message );
     void HandleDeviceDisconnected( const FIOSLaunchDaemonPong& Message );
@@ -193,7 +193,7 @@ private:
 	FDelegateHandle TickDelegateHandle;
 
 	// Holds the message endpoint used for communicating with the LaunchDaemon.
-	FMessageEndpointPtr MessageEndpoint;
+	TSharedPtr<FMessageEndpoint, ESPMode::ThreadSafe> MessageEndpoint;
 
 #if WITH_ENGINE
 	// Holds the Engine INI settings, for quick use.
diff --git a/Engine/Source/Developer/IOS/IOSTargetPlatform/Private/IOSTargetPlatformModule.cpp b/Engine/Source/Developer/IOS/IOSTargetPlatform/Private/IOSTargetPlatformModule.cpp
--- a/Engine/Source/Developer/IOS/IOSTargetPlatform/Private/IOSTargetPlatformModule.cpp
+++ b/Engine/Source/Developer/IOS/IOSTargetPlatform/Private/IOSTargetPlatformModule.cpp
@@ -2,36 +2,36 @@
 
 #include "IOSTargetPlatform.h"
 #include "Interfaces/ITargetPlatformModule.h"
+#include "Modules/ModuleManager.h"
+
 
 /**
  * Holds the target platform singleton.
  */
-static ITargetPlatform* Singleton = NULL;
+static ITargetPlatform* Singleton = nullptr;
 
 
 /**
  * Module for iOS as a target platform
  */
-class FIOSTargetPlatformModule : public ITargetPlatformModule
+class FIOSTargetPlatformModule
+	: public ITargetPlatformModule
 {
 public:
 
-	/**
-	 * Destructor.
-	 */
+	/** Destructor. */
 	~FIOSTargetPlatformModule()
 	{
-		Singleton = NULL;
+		Singleton = nullptr;
 	}
 
-
 public:
 
-	// Begin ITargetPlatformModule interface
+	//~ ITargetPlatformModule interface
 
 	virtual ITargetPlatform* GetTargetPlatform()
 	{
-		if (Singleton == NULL)
+		if (Singleton == nullptr)
 		{
 			Singleton = new FIOSTargetPlatform(false);
 		}
@@ -39,20 +39,13 @@ public:
 		return Singleton;
 	}
 
-	// End ITargetPlatformModule interface
-
-
 public:
 
-	// Begin IModuleInterface interface
-	virtual void StartupModule() override
-	{
-	}
+	//~ IModuleInterface interface
 
-	virtual void ShutdownModule() override
-	{
-	}
-	// End IModuleInterface interface
+	virtual void StartupModule() override { }
+	virtual void ShutdownModule() override { }
 };
 
-IMPLEMENT_MODULE( FIOSTargetPlatformModule, IOSTargetPlatform);
+
+IMPLEMENT_MODULE(FIOSTargetPlatformModule, IOSTargetPlatform);
diff --git a/Engine/Source/Developer/IOS/IOSTargetPlatform/Private/Windows/IOSDeviceHelperWindows.cpp b/Engine/Source/Developer/IOS/IOSTargetPlatform/Private/Windows/IOSDeviceHelperWindows.cpp
--- a/Engine/Source/Developer/IOS/IOSTargetPlatform/Private/Windows/IOSDeviceHelperWindows.cpp
+++ b/Engine/Source/Developer/IOS/IOSTargetPlatform/Private/Windows/IOSDeviceHelperWindows.cpp
@@ -1,6 +1,7 @@
 // Copyright 1998-2017 Epic Games, Inc. All Rights Reserved
 
 #include "IOSTargetPlatform.h"
+#include "HAL/PlatformProcess.h"
 #include "HAL/Runnable.h"
 #include "HAL/RunnableThread.h"
 
@@ -15,7 +16,7 @@ class FDeviceSyslogRelay : public FRunnable
 {
 public:
 
-	FDeviceSyslogRelay(FString const&	inDeviceID)
+	FDeviceSyslogRelay(FString const& inDeviceID)
 	{
 		deviceID = inDeviceID;
 		Stopping = false;
diff --git a/Engine/Source/Developer/IOS/TVOSTargetPlatform/Private/TVOSTargetPlatformModule.cpp b/Engine/Source/Developer/IOS/TVOSTargetPlatform/Private/TVOSTargetPlatformModule.cpp
--- a/Engine/Source/Developer/IOS/TVOSTargetPlatform/Private/TVOSTargetPlatformModule.cpp
+++ b/Engine/Source/Developer/IOS/TVOSTargetPlatform/Private/TVOSTargetPlatformModule.cpp
@@ -1,19 +1,19 @@
 // Copyright 1998-2017 Epic Games, Inc. All Rights Reserved.
 
-#include "CoreMinimal.h"
 #include "IOSTargetPlatform.h"
 #include "Interfaces/ITargetPlatformModule.h"
+#include "Modules/ModuleManager.h"
 
-/**
- * Holds the target platform singleton.
- */
-static ITargetPlatform* Singleton = NULL;
+
+/** Holds the target platform singleton. */
+static ITargetPlatform* Singleton = nullptr;
 
 
 /**
  * Module for TVOS as a target platform
  */
-class FTVOSTargetPlatformModule : public ITargetPlatformModule
+class FTVOSTargetPlatformModule
+	: public ITargetPlatformModule
 {
 public:
 
@@ -22,17 +22,16 @@ public:
 	 */
 	~FTVOSTargetPlatformModule()
 	{
-		Singleton = NULL;
+		Singleton = nullptr;
 	}
 
-
 public:
 
-	// Begin ITargetPlatformModule interface
+	//~ ITargetPlatformModule interface
 
 	virtual ITargetPlatform* GetTargetPlatform()
 	{
-		if (Singleton == NULL)
+		if (Singleton == nullptr)
 		{
 			Singleton = new FIOSTargetPlatform(true);
 		}
@@ -40,20 +39,13 @@ public:
 		return Singleton;
 	}
 
-	// End ITargetPlatformModule interface
-
-
 public:
 
-	// Begin IModuleInterface interface
-	virtual void StartupModule() override
-	{
-	}
+	//~ IModuleInterface interface
 
-	virtual void ShutdownModule() override
-	{
-	}
-	// End IModuleInterface interface
+	virtual void StartupModule() override { }
+	virtual void ShutdownModule() override { }
 };
 
-IMPLEMENT_MODULE( FTVOSTargetPlatformModule, TVOSTargetPlatform);
+
+IMPLEMENT_MODULE(FTVOSTargetPlatformModule, TVOSTargetPlatform);
diff --git a/Engine/Source/Developer/LauncherServices/Private/Launcher/Launcher.cpp b/Engine/Source/Developer/LauncherServices/Private/Launcher/Launcher.cpp
--- a/Engine/Source/Developer/LauncherServices/Private/Launcher/Launcher.cpp
+++ b/Engine/Source/Developer/LauncherServices/Private/Launcher/Launcher.cpp
@@ -1,18 +1,21 @@
 // Copyright 1998-2017 Epic Games, Inc. All Rights Reserved.
 
 #include "Launcher/Launcher.h"
+
 #include "HAL/RunnableThread.h"
 #include "Launcher/LauncherWorker.h"
 
+
 /* Static class member instantiations
 *****************************************************************************/
 
 FThreadSafeCounter FLauncher::WorkerCounter;
 
+
 /* ILauncher overrides
  *****************************************************************************/
 
-ILauncherWorkerPtr FLauncher::Launch( const ITargetDeviceProxyManagerRef& DeviceProxyManager, const ILauncherProfileRef& Profile )
+ILauncherWorkerPtr FLauncher::Launch(const TSharedRef<ITargetDeviceProxyManager>& DeviceProxyManager, const ILauncherProfileRef& Profile)
 {
 	if (Profile->IsValidForLaunch())
 	{
diff --git a/Engine/Source/Developer/LauncherServices/Private/Launcher/Launcher.h b/Engine/Source/Developer/LauncherServices/Private/Launcher/Launcher.h
--- a/Engine/Source/Developer/LauncherServices/Private/Launcher/Launcher.h
+++ b/Engine/Source/Developer/LauncherServices/Private/Launcher/Launcher.h
@@ -4,7 +4,8 @@
 
 #include "CoreMinimal.h"
 #include "HAL/ThreadSafeCounter.h"
-#include "Interfaces/ILauncher.h"
+#include "ILauncher.h"
+
 
 /**
  * Implements the game launcher.
@@ -14,14 +15,12 @@ class FLauncher
 {
 public:
 
-	//~ Begin ILauncher Interface
-
-	virtual ILauncherWorkerPtr Launch( const ITargetDeviceProxyManagerRef& DeviceProxyManager, const ILauncherProfileRef& Profile ) override;
+	//~ ILauncher Interface
 
-	//~ End ILauncher Interface
+	virtual ILauncherWorkerPtr Launch(const TSharedRef<ITargetDeviceProxyManager>& DeviceProxyManager, const ILauncherProfileRef& Profile) override;
 
 private:
 
-	// Worker counter, used to generate unique thread names for each worker
+	/** Worker counter, used to generate unique thread names for each worker. */
 	static FThreadSafeCounter WorkerCounter;
 };
diff --git a/Engine/Source/Developer/LauncherServices/Private/Launcher/LauncherTask.h b/Engine/Source/Developer/LauncherServices/Private/Launcher/LauncherTask.h
--- a/Engine/Source/Developer/LauncherServices/Private/Launcher/LauncherTask.h
+++ b/Engine/Source/Developer/LauncherServices/Private/Launcher/LauncherTask.h
@@ -4,7 +4,7 @@
 
 #include "CoreMinimal.h"
 #include "HAL/ThreadSafeCounter.h"
-#include "Interfaces/ILauncherTask.h"
+#include "ILauncherTask.h"
 #include "HAL/PlatformProcess.h"
 #include "HAL/Runnable.h"
 #include "Launcher/LauncherTaskChainState.h"
diff --git a/Engine/Source/Developer/LauncherServices/Private/Launcher/LauncherTaskChainState.h b/Engine/Source/Developer/LauncherServices/Private/Launcher/LauncherTaskChainState.h
--- a/Engine/Source/Developer/LauncherServices/Private/Launcher/LauncherTaskChainState.h
+++ b/Engine/Source/Developer/LauncherServices/Private/Launcher/LauncherTaskChainState.h
@@ -4,7 +4,7 @@
 
 #include "CoreMinimal.h"
 #include "Misc/Guid.h"
-#include "Interfaces/ILauncherProfile.h"
+#include "ILauncherProfile.h"
 
 /**
  * Structure for passing launcher task chain state data.
diff --git a/Engine/Source/Developer/LauncherServices/Private/Launcher/LauncherUATTask.h b/Engine/Source/Developer/LauncherServices/Private/Launcher/LauncherUATTask.h
--- a/Engine/Source/Developer/LauncherServices/Private/Launcher/LauncherUATTask.h
+++ b/Engine/Source/Developer/LauncherServices/Private/Launcher/LauncherUATTask.h
@@ -3,7 +3,7 @@
 #pragma once
 
 #include "CoreMinimal.h"
-#include "Interfaces/ILauncherWorker.h"
+#include "ILauncherWorker.h"
 #include "Misc/CommandLine.h"
 #include "Misc/Paths.h"
 #include "Launcher/LauncherTask.h"
diff --git a/Engine/Source/Developer/LauncherServices/Private/Launcher/LauncherWorker.cpp b/Engine/Source/Developer/LauncherServices/Private/Launcher/LauncherWorker.cpp
--- a/Engine/Source/Developer/LauncherServices/Private/Launcher/LauncherWorker.cpp
+++ b/Engine/Source/Developer/LauncherServices/Private/Launcher/LauncherWorker.cpp
@@ -3,6 +3,10 @@
 #include "Launcher/LauncherWorker.h"
 #include "HAL/PlatformTime.h"
 #include "HAL/FileManager.h"
+#include "ISourceCodeAccessor.h"
+#include "ISourceCodeAccessModule.h"
+#include "ITargetDeviceProxy.h"
+#include "ITargetDeviceProxyManager.h"
 #include "Misc/CommandLine.h"
 #include "Misc/Paths.h"
 #include "HAL/ThreadSafeCounter.h"
@@ -13,20 +17,21 @@
 #include "Launcher/LauncherUATTask.h"
 #include "Launcher/LauncherVerifyProfileTask.h"
 #include "PlatformInfo.h"
-#include "ISourceCodeAccessor.h"
-#include "ISourceCodeAccessModule.h"
+
 
 #define LOCTEXT_NAMESPACE "LauncherWorker"
 
+
 /* Static class member instantiations
 *****************************************************************************/
 
 FThreadSafeCounter FLauncherTask::TaskCounter;
 
+
 /* FLauncherWorker structors
  *****************************************************************************/
 
-FLauncherWorker::FLauncherWorker( const ITargetDeviceProxyManagerRef& InDeviceProxyManager, const ILauncherProfileRef& InProfile )
+FLauncherWorker::FLauncherWorker(const TSharedRef<ITargetDeviceProxyManager>& InDeviceProxyManager, const ILauncherProfileRef& InProfile)
 	: DeviceProxyManager(InDeviceProxyManager)
 	, Profile(InProfile)
 	, Status(ELauncherWorkerStatus::Busy)
@@ -211,7 +216,7 @@ void FLauncherWorker::OnTaskCompleted(const FString& TaskName)
 	StageCompleted.Broadcast(TaskName, FPlatformTime::Seconds() - StageStartTime);
 }
 
-static void AddDeviceToLaunchCommand(const FString& DeviceId, ITargetDeviceProxyPtr DeviceProxy, const ILauncherProfileRef& InProfile, FString& DeviceNames, FString& RoleCommands, bool& bVsyncAdded)
+static void AddDeviceToLaunchCommand(const FString& DeviceId, TSharedPtr<ITargetDeviceProxy> DeviceProxy, const ILauncherProfileRef& InProfile, FString& DeviceNames, FString& RoleCommands, bool& bVsyncAdded)
 {
 	// add the platform
 	DeviceNames += TEXT("+\"") + DeviceId + TEXT("\"");
@@ -392,7 +397,7 @@ FString FLauncherWorker::CreateUATCommand( const ILauncherProfileRef& InProfile,
 		for (int32 DeviceIndex = 0; DeviceIndex < Devices.Num(); ++DeviceIndex)
 		{
 			const FString& DeviceId = Devices[DeviceIndex];
-			ITargetDeviceProxyPtr DeviceProxy = DeviceProxyManager->FindProxyDeviceForTargetDevice(DeviceId);
+			TSharedPtr<ITargetDeviceProxy> DeviceProxy = DeviceProxyManager->FindProxyDeviceForTargetDevice(DeviceId);
 			if (DeviceProxy.IsValid())
 			{
 				AddDeviceToLaunchCommand(DeviceId, DeviceProxy, InProfile, DeviceNames, RoleCommands, bVsyncAdded);
@@ -801,7 +806,7 @@ void FLauncherWorker::CreateAndExecuteTasks( const ILauncherProfileRef& InProfil
 		{
 			const FString& DeviceId = Devices[DeviceIndex];
 
-			ITargetDeviceProxyPtr DeviceProxy = DeviceProxyManager->FindProxyDeviceForTargetDevice(DeviceId);
+			TSharedPtr<ITargetDeviceProxy> DeviceProxy = DeviceProxyManager->FindProxyDeviceForTargetDevice(DeviceId);
 
 			if (DeviceProxy.IsValid())
 			{
diff --git a/Engine/Source/Developer/LauncherServices/Private/Launcher/LauncherWorker.h b/Engine/Source/Developer/LauncherServices/Private/Launcher/LauncherWorker.h
--- a/Engine/Source/Developer/LauncherServices/Private/Launcher/LauncherWorker.h
+++ b/Engine/Source/Developer/LauncherServices/Private/Launcher/LauncherWorker.h
@@ -3,12 +3,13 @@
 #pragma once
 
 #include "CoreMinimal.h"
-#include "Interfaces/ITargetDeviceProxyManager.h"
-#include "Interfaces/ILauncherWorker.h"
+#include "ILauncherWorker.h"
 #include "HAL/PlatformProcess.h"
 #include "HAL/Runnable.h"
 
 class FLauncherTask;
+class ITargetDeviceProxyManager;
+
 
 struct FCommandDesc
 {
@@ -29,10 +30,10 @@ public:
 	/**
 	 * Creates and initializes a new instance.
 	 *
-	 * @param InProfile - The profile to process.
-	 * @param InDeviceProxyManager - The target device proxy manager to use.
+	 * @param InProfile The profile to process.
+	 * @param InDeviceProxyManager The target device proxy manager to use.
 	 */
-	FLauncherWorker( const ITargetDeviceProxyManagerRef& InDeviceProxyManager, const ILauncherProfileRef& InProfile );
+	FLauncherWorker(const TSharedRef<ITargetDeviceProxyManager>& InDeviceProxyManager, const ILauncherProfileRef& InProfile);
 
 	~FLauncherWorker()
 	{
@@ -112,7 +113,7 @@ private:
 	FCriticalSection CriticalSection;
 
 	// Holds a pointer  to the device proxy manager.
-	ITargetDeviceProxyManagerPtr DeviceProxyManager;
+	TSharedPtr<ITargetDeviceProxyManager> DeviceProxyManager;
 
 	// Holds a pointer to the launcher profile.
 	ILauncherProfilePtr Profile;
diff --git a/Engine/Source/Developer/LauncherServices/Private/LauncherServicesModule.cpp b/Engine/Source/Developer/LauncherServices/Private/LauncherServicesModule.cpp
--- a/Engine/Source/Developer/LauncherServices/Private/LauncherServicesModule.cpp
+++ b/Engine/Source/Developer/LauncherServices/Private/LauncherServicesModule.cpp
@@ -3,7 +3,7 @@
 #include "CoreMinimal.h"
 #include "Misc/Guid.h"
 #include "Modules/ModuleManager.h"
-#include "Interfaces/ILauncherServicesModule.h"
+#include "ILauncherServicesModule.h"
 #include "Profiles/LauncherDeviceGroup.h"
 #include "Profiles/LauncherProfile.h"
 #include "Profiles/LauncherProfileManager.h"
diff --git a/Engine/Source/Developer/LauncherServices/Private/Profiles/LauncherDeviceGroup.h b/Engine/Source/Developer/LauncherServices/Private/Profiles/LauncherDeviceGroup.h
--- a/Engine/Source/Developer/LauncherServices/Private/Profiles/LauncherDeviceGroup.h
+++ b/Engine/Source/Developer/LauncherServices/Private/Profiles/LauncherDeviceGroup.h
@@ -4,7 +4,7 @@
 
 #include "CoreMinimal.h"
 #include "Misc/Guid.h"
-#include "Interfaces/ILauncherDeviceGroup.h"
+#include "ILauncherDeviceGroup.h"
 
 /**
  * Implements a group of devices for the Launcher user interface.
diff --git a/Engine/Source/Developer/LauncherServices/Private/Profiles/LauncherProfile.h b/Engine/Source/Developer/LauncherServices/Private/Profiles/LauncherProfile.h
--- a/Engine/Source/Developer/LauncherServices/Private/Profiles/LauncherProfile.h
+++ b/Engine/Source/Developer/LauncherServices/Private/Profiles/LauncherProfile.h
@@ -5,8 +5,8 @@
 #include "CoreMinimal.h"
 #include "Misc/Guid.h"
 #include "Modules/ModuleManager.h"
-#include "Interfaces/ILauncherProfile.h"
-#include "Interfaces/ILauncherServicesModule.h"
+#include "ILauncherProfile.h"
+#include "ILauncherServicesModule.h"
 #include "Misc/Paths.h"
 #include "Launcher/LauncherProjectPath.h"
 #include "Misc/CommandLine.h"
@@ -14,8 +14,9 @@
 #include "Misc/App.h"
 #include "Interfaces/ITargetPlatform.h"
 #include "Interfaces/ITargetPlatformManagerModule.h"
-#include "Interfaces/ITargetDeviceServicesModule.h"
-#include "Common/GameProjectHelper.h"
+#include "ITargetDeviceProxy.h"
+#include "ITargetDeviceServicesModule.h"
+#include "GameProjectHelper.h"
 #include "Profiles/LauncherProfileLaunchRole.h"
 #include "PlatformInfo.h"
 
@@ -295,16 +296,16 @@ public:
 
 			if (DefaultDeployPlatform != NAME_None)
 			{
-				TArray<ITargetDeviceProxyPtr> PlatformDeviceProxies;
+				TArray<TSharedPtr<ITargetDeviceProxy>> PlatformDeviceProxies;
 				ITargetDeviceServicesModule& TargetDeviceServicesModule = FModuleManager::LoadModuleChecked<ITargetDeviceServicesModule>("TargetDeviceServices");
-				const ITargetDeviceProxyManagerRef& InDeviceProxyManager = TargetDeviceServicesModule.GetDeviceProxyManager();
+				const TSharedRef<ITargetDeviceProxyManager>& InDeviceProxyManager = TargetDeviceServicesModule.GetDeviceProxyManager();
 
 				InDeviceProxyManager->GetProxies(NAME_None, true, PlatformDeviceProxies);
 
-				ITargetDeviceProxyPtr DefaultPlatformDevice;
+				TSharedPtr<ITargetDeviceProxy> DefaultPlatformDevice;
 				for (int32 ProxyIndex = 0; ProxyIndex < PlatformDeviceProxies.Num(); ++ProxyIndex)
 				{
-					ITargetDeviceProxyPtr DeviceProxy = PlatformDeviceProxies[ProxyIndex];
+					TSharedPtr<ITargetDeviceProxy> DeviceProxy = PlatformDeviceProxies[ProxyIndex];
 
 					if (DeviceProxy->GetVanillaPlatformId(NAME_None) == DefaultDeployPlatform)
 					{
@@ -1144,7 +1145,7 @@ public:
 
 		// devices
 		ITargetDeviceServicesModule& DeviceServiceModule = FModuleManager::LoadModuleChecked<ITargetDeviceServicesModule>(TEXT("TargetDeviceServices"));
-		ITargetDeviceProxyManagerRef DeviceProxyManager = DeviceServiceModule.GetDeviceProxyManager();
+		TSharedRef<ITargetDeviceProxyManager> DeviceProxyManager = DeviceServiceModule.GetDeviceProxyManager();
 		ILauncherDeviceGroupPtr DeviceGroup = GetDeployedDeviceGroup();
 		TMap<FString, FString> RoleCommands;
 		FString CommandLine = GetDefaultLaunchRole()->GetUATCommandLine();
@@ -1165,7 +1166,7 @@ public:
 			for (int32 DeviceIndex = 0; DeviceIndex < Devices.Num(); ++DeviceIndex)
 			{
 				const FString& DeviceId = Devices[DeviceIndex];
-				ITargetDeviceProxyPtr DeviceProxy = DeviceProxyManager->FindProxyDeviceForTargetDevice(DeviceId);
+				TSharedPtr<ITargetDeviceProxy> DeviceProxy = DeviceProxyManager->FindProxyDeviceForTargetDevice(DeviceId);
 				Writer.WriteValue(DeviceId);
 				if (DeviceProxy.IsValid())
 				{
@@ -1504,7 +1505,7 @@ public:
 
 		// Loading the Device Proxy Manager to get the needed Device Manager.
 		ITargetDeviceServicesModule& DeviceServiceModule = FModuleManager::LoadModuleChecked<ITargetDeviceServicesModule>(TEXT("TargetDeviceServices"));
-		ITargetDeviceProxyManagerRef DeviceProxyManager = DeviceServiceModule.GetDeviceProxyManager();
+		TSharedRef<ITargetDeviceProxyManager> DeviceProxyManager = DeviceServiceModule.GetDeviceProxyManager();
 
 		if (DeviceGroup.IsValid() && Platforms.Num() < 1)
 		{
@@ -1514,7 +1515,7 @@ public:
 			{
 				const FString& DeviceId = Devices[DeviceIndex];
 
-				ITargetDeviceProxyPtr DeviceProxy = DeviceProxyManager->FindProxyDeviceForTargetDevice(DeviceId);
+				TSharedPtr<ITargetDeviceProxy> DeviceProxy = DeviceProxyManager->FindProxyDeviceForTargetDevice(DeviceId);
 
 				if (DeviceProxy.IsValid())
 				{
@@ -2483,7 +2484,7 @@ protected:
 				
 				if (TargetDeviceServicesModule)
 				{
-					ITargetDeviceProxyPtr DeviceProxy = TargetDeviceServicesModule->GetDeviceProxyManager()->FindProxy(DeviceId);
+					TSharedPtr<ITargetDeviceProxy> DeviceProxy = TargetDeviceServicesModule->GetDeviceProxyManager()->FindProxy(DeviceId);
 					
 					if(DeviceProxy.IsValid())
 					{
diff --git a/Engine/Source/Developer/LauncherServices/Private/Profiles/LauncherProfileLaunchRole.h b/Engine/Source/Developer/LauncherServices/Private/Profiles/LauncherProfileLaunchRole.h
--- a/Engine/Source/Developer/LauncherServices/Private/Profiles/LauncherProfileLaunchRole.h
+++ b/Engine/Source/Developer/LauncherServices/Private/Profiles/LauncherProfileLaunchRole.h
@@ -3,7 +3,7 @@
 #pragma once
 
 #include "CoreMinimal.h"
-#include "Interfaces/ILauncherProfileLaunchRole.h"
+#include "ILauncherProfileLaunchRole.h"
 #include "Dom/JsonObject.h"
 
 class FLauncherProfileLaunchRole
diff --git a/Engine/Source/Developer/Linux/LinuxTargetPlatform/Private/LinuxTargetDevice.h b/Engine/Source/Developer/Linux/LinuxTargetPlatform/Private/LinuxTargetDevice.h
--- a/Engine/Source/Developer/Linux/LinuxTargetPlatform/Private/LinuxTargetDevice.h
+++ b/Engine/Source/Developer/Linux/LinuxTargetPlatform/Private/LinuxTargetDevice.h
@@ -1,18 +1,14 @@
 // Copyright 1998-2017 Epic Games, Inc. All Rights Reserved.
 
-/*=============================================================================
-	LocalLinuxTargetDevice.h: Declares the LocalLinuxTargetDevice class.
-=============================================================================*/
-
 #pragma once
 
 #include "CoreMinimal.h"
-#include "Interfaces/TargetDeviceId.h"
-#include "Interfaces/ITargetDevice.h"
 #include "Misc/Paths.h"
 #include "HAL/FileManager.h"
 #include "HAL/PlatformProcess.h"
 #include "ITargetPlatform.h"
+#include "Interfaces/ITargetDevice.h"
+#include "Interfaces/TargetDeviceId.h"
 
 #if PLATFORM_LINUX
 	#include <signal.h>
@@ -23,6 +19,7 @@ class FLinuxTargetDevice;
 class IFileManager;
 struct FProcHandle;
 
+
 /**
  * Type definition for shared pointers to instances of FLinuxTargetDevice.
  */
diff --git a/Engine/Source/Developer/OutputLog/Private/SDeviceOutputLog.h b/Engine/Source/Developer/OutputLog/Private/SDeviceOutputLog.h
--- a/Engine/Source/Developer/OutputLog/Private/SDeviceOutputLog.h
+++ b/Engine/Source/Developer/OutputLog/Private/SDeviceOutputLog.h
@@ -7,7 +7,7 @@
 #include "Widgets/SWidget.h"
 #include "Widgets/DeclarativeSyntaxSupport.h"
 #include "SOutputLog.h"
-#include "Interfaces/TargetDeviceId.h"
+#include "TargetDeviceId.h"
 #include "Interfaces/ITargetDevice.h"
 
 class SComboButton;
diff --git a/Engine/Source/Developer/ProfilerClient/Private/ProfilerClientManager.cpp b/Engine/Source/Developer/ProfilerClient/Private/ProfilerClientManager.cpp
--- a/Engine/Source/Developer/ProfilerClient/Private/ProfilerClientManager.cpp
+++ b/Engine/Source/Developer/ProfilerClient/Private/ProfilerClientManager.cpp
@@ -2,11 +2,11 @@
 
 #include "ProfilerClientManager.h"
 #include "HAL/FileManager.h"
+#include "MessageEndpointBuilder.h"
 #include "Misc/Paths.h"
+#include "Misc/SecureHash.h"
 #include "Serialization/MemoryReader.h"
-#include "Helpers/MessageEndpointBuilder.h"
 #include "ProfilerServiceMessages.h"
-#include "Misc/SecureHash.h"
 
 
 DEFINE_LOG_CATEGORY_STATIC(LogProfilerClient, Log, All);
diff --git a/Engine/Source/Developer/ProfilerClient/Private/ProfilerClientManager.h b/Engine/Source/Developer/ProfilerClient/Private/ProfilerClientManager.h
--- a/Engine/Source/Developer/ProfilerClient/Private/ProfilerClientManager.h
+++ b/Engine/Source/Developer/ProfilerClient/Private/ProfilerClientManager.h
@@ -10,7 +10,7 @@
 #include "Containers/Ticker.h"
 #include "IProfilerServiceManager.h"
 #include "IMessageBus.h"
-#include "Helpers/MessageEndpoint.h"
+#include "MessageEndpoint.h"
 #include "Stats/StatsFile.h"
 
 class FNewStatsReader;
@@ -21,6 +21,7 @@ struct FProfilerServiceFileChunk;
 struct FProfilerServicePing;
 struct FProfilerServicePreviewAck;
 
+
 /**
  * Helper struct containing all of the data and operations associated with a service connection.
  */
@@ -235,7 +236,7 @@ private:
 	TSet<FString> FailedTransfer;
 
 	/** Holds the messaging endpoint. */
-	FMessageEndpointPtr MessageEndpoint;
+	TSharedPtr<FMessageEndpoint, ESPMode::ThreadSafe> MessageEndpoint;
 
 	/** Holds a pointer to the message bus. */
 	TSharedPtr<IMessageBus, ESPMode::ThreadSafe> MessageBus;
diff --git a/Engine/Source/Developer/ProfilerService/Private/ProfilerServiceFileTransfer.cpp b/Engine/Source/Developer/ProfilerService/Private/ProfilerServiceFileTransfer.cpp
--- a/Engine/Source/Developer/ProfilerService/Private/ProfilerServiceFileTransfer.cpp
+++ b/Engine/Source/Developer/ProfilerService/Private/ProfilerServiceFileTransfer.cpp
@@ -12,7 +12,7 @@
 #include "ProfilerServiceManager.h"
 
 
-FFileTransferRunnable::FFileTransferRunnable(FMessageEndpointPtr& InMessageEndpoint)
+FFileTransferRunnable::FFileTransferRunnable(TSharedPtr<FMessageEndpoint, ESPMode::ThreadSafe>& InMessageEndpoint)
 	: Runnable(nullptr)
 	, WorkEvent(FPlatformProcess::GetSynchEventFromPool(true))
 	, MessageEndpoint(InMessageEndpoint)
diff --git a/Engine/Source/Developer/ProfilerService/Private/ProfilerServiceManager.cpp b/Engine/Source/Developer/ProfilerService/Private/ProfilerServiceManager.cpp
--- a/Engine/Source/Developer/ProfilerService/Private/ProfilerServiceManager.cpp
+++ b/Engine/Source/Developer/ProfilerService/Private/ProfilerServiceManager.cpp
@@ -1,10 +1,10 @@
 // Copyright 1998-2017 Epic Games, Inc. All Rights Reserved.
 
 #include "ProfilerServiceManager.h"
-#include "Serialization/MemoryReader.h"
+#include "MessageEndpointBuilder.h"
 #include "Misc/App.h"
-#include "Helpers/MessageEndpointBuilder.h"
 #include "ProfilerServiceMessages.h"
+#include "Serialization/MemoryReader.h"
 #include "Stats/StatsData.h"
 #include "Stats/StatsFile.h"
 
diff --git a/Engine/Source/Developer/ProfilerService/Private/ProfilerServiceManager.h b/Engine/Source/Developer/ProfilerService/Private/ProfilerServiceManager.h
--- a/Engine/Source/Developer/ProfilerService/Private/ProfilerServiceManager.h
+++ b/Engine/Source/Developer/ProfilerService/Private/ProfilerServiceManager.h
@@ -10,7 +10,7 @@
 #include "HAL/Runnable.h"
 #include "Containers/Queue.h"
 #include "Containers/Ticker.h"
-#include "Helpers/MessageEndpoint.h"
+#include "MessageEndpoint.h"
 
 struct FProfilerServiceCapture;
 struct FProfilerServiceData2;
@@ -88,7 +88,7 @@ protected:
 	FEvent* WorkEvent;
 
 	/** Holds the messaging endpoint. */
-	FMessageEndpointPtr& MessageEndpoint;
+	TSharedPtr<FMessageEndpoint, ESPMode::ThreadSafe> MessageEndpoint;
 
 	/** > 0 if we have been asked to abort work in progress at the next opportunity. */
 	FThreadSafeCounter StopTaskCounter;
@@ -210,7 +210,7 @@ private:
 	void RemoveNewFrameHandleStatsThread();
 
 	/** Holds the messaging endpoint. */
-	FMessageEndpointPtr MessageEndpoint;
+	TSharedPtr<FMessageEndpoint, ESPMode::ThreadSafe> MessageEndpoint;
 
 	/** Holds the session and instance identifier. */
 	FGuid SessionId;
diff --git a/Engine/Source/Developer/ProjectLauncher/Private/Models/ProjectLauncherCommands.h b/Engine/Source/Developer/ProjectLauncher/Private/Models/ProjectLauncherCommands.h
--- a/Engine/Source/Developer/ProjectLauncher/Private/Models/ProjectLauncherCommands.h
+++ b/Engine/Source/Developer/ProjectLauncher/Private/Models/ProjectLauncherCommands.h
@@ -2,39 +2,40 @@
 
 #pragma once
 
-#include "CoreMinimal.h"
+#include "EditorStyleSet.h"
 #include "InputCoreTypes.h"
 #include "Framework/Commands/InputChord.h"
 #include "Framework/Commands/Commands.h"
-#include "EditorStyleSet.h"
+#include "Templates/SharedPointer.h"
+
 
 #define LOCTEXT_NAMESPACE "ProjectLauncherCommands"
 
+
 /**
  * The project launcher UI commands
  */
-class FProjectLauncherCommands : public TCommands<FProjectLauncherCommands>
+class FProjectLauncherCommands
+	: public TCommands<FProjectLauncherCommands>
 {
 public:
 
-	/**
-	 * Default constructor.
-	 */
-	FProjectLauncherCommands( )
+	/** Default constructor. */
+	FProjectLauncherCommands()
 		: TCommands<FProjectLauncherCommands>(
 			"LauncherCommand",
 			NSLOCTEXT("Contexts", "LauncherCommand", "Project Launcher Command"),
 			NAME_None,
 			FEditorStyle::GetStyleSetName()
-	)
+		)
 	{ }
 	
 public:
 
-	// TCommands interface
+	//~ TCommands interface
 
 	PRAGMA_DISABLE_OPTIMIZATION
-	virtual void RegisterCommands( ) override
+	virtual void RegisterCommands() override
 	{
 		UI_COMMAND(QuickLaunch, "Quick Launch", "Builds, cooks, and launches a build.", EUserInterfaceActionType::ToggleButton, FInputChord(EModifierKey::Control, EKeys::L));
 		UI_COMMAND(CreateBuild, "Build", "Creates a build.", EUserInterfaceActionType::ToggleButton, FInputChord(EModifierKey::Control, EKeys::B));
@@ -50,20 +51,14 @@ public:
 public:
 
 	TSharedPtr<FUICommandInfo> CreateBuild;
-
 	TSharedPtr<FUICommandInfo> DeployBuild;
-
 	TSharedPtr<FUICommandInfo> QuickLaunch;
-
 	TSharedPtr<FUICommandInfo> AdvancedBuild;
-
 	TSharedPtr<FUICommandInfo> CloseSettings;
-
 	TSharedPtr<FUICommandInfo> RenameProfile;
-
 	TSharedPtr<FUICommandInfo> DuplicateProfile;
-
 	TSharedPtr<FUICommandInfo> DeleteProfile;
 };
 
+
 #undef LOCTEXT_NAMESPACE
diff --git a/Engine/Source/Developer/ProjectLauncher/Private/Models/ProjectLauncherModel.h b/Engine/Source/Developer/ProjectLauncher/Private/Models/ProjectLauncherModel.h
--- a/Engine/Source/Developer/ProjectLauncher/Private/Models/ProjectLauncherModel.h
+++ b/Engine/Source/Developer/ProjectLauncher/Private/Models/ProjectLauncherModel.h
@@ -2,17 +2,22 @@
 
 #pragma once
 
-#include "CoreMinimal.h"
+#include "CoreTypes.h"
+#include "Containers/Array.h"
+#include "Delegates/Delegate.h"
+#include "GameProjectHelper.h"
+#include "ILauncher.h"
+#include "ILauncherProfileManager.h"
 #include "Misc/Paths.h"
 #include "Misc/ConfigCacheIni.h"
 #include "Misc/App.h"
-#include "Interfaces/ILauncherProfileManager.h"
-#include "Interfaces/ILauncher.h"
-#include "Common/GameProjectHelper.h"
+#include "Templates/SharedPointer.h"
+
 
 class FProjectLauncherModel;
 class SProjectLauncher;
 
+
 namespace ELauncherPanels
 {
 	/**
@@ -29,18 +34,11 @@ namespace ELauncherPanels
 }
 
 
-/** Type definition for shared pointers to instances of FProjectLauncherModel. */
-typedef TSharedPtr<class FProjectLauncherModel> FProjectLauncherModelPtr;
-
-/** Type definition for shared references to instances of FProjectLauncherModel. */
-typedef TSharedRef<class FProjectLauncherModel> FProjectLauncherModelRef;
-
-
 /**
  * Delegate type for launcher profile selection changes.
  *
- * The first parameter is the newly selected profile (or NULL if none was selected).
- * The second parameter is the old selected profile (or NULL if none was previous selected).
+ * The first parameter is the newly selected profile (or nullptr if none was selected).
+ * The second parameter is the old selected profile (or nullptr if none was previous selected).
  */
 DECLARE_MULTICAST_DELEGATE_TwoParams(FOnSelectedSProjectLauncherProfileChanged, const ILauncherProfilePtr&, const ILauncherProfilePtr&)
 
@@ -55,11 +53,11 @@ public:
 	/**
 	 * Creates and initializes a new instance.
 	 *
-	 * @param InDeviceProxyManager - The device proxy manager to use.
-	 * @param InLauncher - The launcher to use.
-	 * @param InProfileManager - The profile manager to use.
+	 * @param InDeviceProxyManager The device proxy manager to use.
+	 * @param InLauncher The launcher to use.
+	 * @param InProfileManager The profile manager to use.
 	 */
-	FProjectLauncherModel(const ITargetDeviceProxyManagerRef& InDeviceProxyManager, const ILauncherRef& InLauncher, const ILauncherProfileManagerRef& InProfileManager)
+	FProjectLauncherModel(const TSharedRef<ITargetDeviceProxyManager>& InDeviceProxyManager, const TSharedRef<ILauncher>& InLauncher, const TSharedRef<ILauncherProfileManager>& InProfileManager)
 		: DeviceProxyManager(InDeviceProxyManager)
 		, SProjectLauncher(InLauncher)
 		, ProfileManager(InProfileManager)
@@ -70,9 +68,7 @@ public:
 		LoadConfig();
 	}
 
-	/**
-	 * Destructor.
-	 */
+	/** Destructor. */
 	~FProjectLauncherModel()
 	{
 		ProfileManager->OnProfileAdded().RemoveAll(this);
@@ -88,7 +84,7 @@ public:
 	 *
 	 * @return Device proxy manager.
 	 */
-	const ITargetDeviceProxyManagerRef& GetDeviceProxyManager() const
+	const TSharedRef<ITargetDeviceProxyManager>& GetDeviceProxyManager() const
 	{
 		return DeviceProxyManager;
 	}
@@ -98,7 +94,7 @@ public:
 	 *
 	 * @return SProjectLauncher.
 	 */
-	const ILauncherRef& GetSProjectLauncher() const
+	const TSharedRef<ILauncher>& GetSProjectLauncher() const
 	{
 		return SProjectLauncher;
 	}
@@ -107,21 +103,20 @@ public:
 	 * Get the model's profile manager.
 	 *
 	 * @return Profile manager.
+	 * @see GetProfiles
 	 */
-	const ILauncherProfileManagerRef& GetProfileManager() const
+	const TSharedRef<ILauncherProfileManager>& GetProfileManager() const
 	{
 		return ProfileManager;
 	}
 
-
 	/**
 	 * Gets the active profile.
 	 *
 	 * @return The active profile.
-	 *
-	 * @see GetProfiles
+	 * @see GetProfiles, GetProfileManager, SelectProfile
 	 */
-	const ILauncherProfilePtr& GetSelectedProfile() const
+	const TSharedPtr<ILauncherProfile>& GetSelectedProfile() const
 	{
 		return SelectedProfile;
 	}
@@ -129,15 +124,16 @@ public:
 	/**
 	 * Sets the active profile.
 	 *
-	 * @param InProfile The profile to assign as active,, or NULL to deselect all.
+	 * @param InProfile The profile to assign as active,, or nullptr to deselect all.
+	 * @see GetSelectedProfile
 	 */
-	void SelectProfile(const ILauncherProfilePtr& Profile)
+	void SelectProfile(const TSharedPtr<ILauncherProfile>& Profile)
 	{
 		if (!Profile.IsValid() || ProfileManager->GetAllProfiles().Contains(Profile))
 		{
 			if (Profile != SelectedProfile)
 			{
-				ILauncherProfilePtr& PreviousProfile = SelectedProfile;
+				TSharedPtr<ILauncherProfile>& PreviousProfile = SelectedProfile;
 				SelectedProfile = Profile;
 
 				ProfileSelectedDelegate.Broadcast(Profile, PreviousProfile);
@@ -151,6 +147,7 @@ public:
 	 * Returns a delegate to be invoked when the profile list has been modified.
 	 *
 	 * @return The delegate.
+	 * @see OnProfileSelected
 	 */
 	FSimpleMulticastDelegate& OnProfileListChanged()
 	{
@@ -161,6 +158,7 @@ public:
 	 * Returns a delegate to be invoked when the selected profile changed.
 	 *
 	 * @return The delegate.
+	 * @see OnProfileListChanged
 	 */
 	FOnSelectedSProjectLauncherProfileChanged& OnProfileSelected()
 	{
@@ -171,11 +169,14 @@ protected:
 
 	/*
 	 * Load all profiles from disk.
+	 *
+	 * @see SaveConfig
 	 */
 	void LoadConfig()
 	{
 		// restore the previous project selection
 		FString ProjectPath;
+
 		if (FPaths::IsProjectFilePathSet())
 		{
 			ProjectPath = FPaths::GetProjectFilePath();
@@ -184,19 +185,22 @@ protected:
 		{
 			ProjectPath = FPaths::RootDir() / FApp::GetProjectName() / FApp::GetProjectName() + TEXT(".uproject");
 		}
-		else if (GConfig != NULL)
+		else if (GConfig != nullptr)
 		{
 			GConfig->GetString(TEXT("FProjectLauncherModel"), TEXT("SelectedProjectPath"), ProjectPath, GEngineIni);
 		}
+
 		ProfileManager->SetProjectPath(ProjectPath);
 	}
 
 	/*
 	 * Saves all profiles to disk.
+
+	 * @see LoadConfig
 	 */
 	void SaveConfig()
 	{
-		if (GConfig != NULL && !FPaths::IsProjectFilePathSet() && !FGameProjectHelper::IsGameAvailable(FApp::GetProjectName()))
+		if (GConfig != nullptr && !FPaths::IsProjectFilePathSet() && !FGameProjectHelper::IsGameAvailable(FApp::GetProjectName()))
 		{
 			FString ProjectPath = ProfileManager->GetProjectPath();
 			GConfig->SetString(TEXT("FProjectLauncherModel"), TEXT("SelectedProjectPath"), *ProjectPath, GEngineIni);
@@ -205,22 +209,22 @@ protected:
 
 private:
 
-	// Callback for when a profile was added to the profile manager.
-	void HandleProfileManagerProfileAdded(const ILauncherProfileRef& Profile)
+	/** Callback for when a profile was added to the profile manager. */
+	void HandleProfileManagerProfileAdded(const TSharedRef<ILauncherProfile>& Profile)
 	{
 		ProfileListChangedDelegate.Broadcast();
 
 		SelectProfile(Profile);
 	}
 
-	// Callback for when a profile was removed from the profile manager.
-	void HandleProfileManagerProfileRemoved(const ILauncherProfileRef& Profile)
+	/** Callback for when a profile was removed from the profile manager. */
+	void HandleProfileManagerProfileRemoved(const TSharedRef<ILauncherProfile>& Profile)
 	{
 		ProfileListChangedDelegate.Broadcast();
 
 		if (Profile == SelectedProfile)
 		{
-			const TArray<ILauncherProfilePtr>& Profiles = ProfileManager->GetAllProfiles();
+			const TArray<TSharedPtr<ILauncherProfile>>& Profiles = ProfileManager->GetAllProfiles();
 
 			if (Profiles.Num() > 0)
 			{
@@ -228,34 +232,33 @@ private:
 			}
 			else
 			{
-				SelectProfile(NULL);
+				SelectProfile(nullptr);
 			}
 		}
 	}
 
 private:
 
-	// Holds a pointer to the device proxy manager.
-	ITargetDeviceProxyManagerRef DeviceProxyManager;
+	/** Pointer to the device proxy manager. */
+	TSharedRef<ITargetDeviceProxyManager> DeviceProxyManager;
 
-	// Holds a pointer to the launcher.
-	ILauncherRef SProjectLauncher;
+	/** Pointer to the launcher. */
+	TSharedRef<ILauncher> SProjectLauncher;
 
-	// Holds a pointer to the profile manager.
-	ILauncherProfileManagerRef ProfileManager;
+	/** Pointer to the profile manager. */
+	TSharedRef<ILauncherProfileManager> ProfileManager;
 
-	// Holds a pointer to active profile.
-	ILauncherProfilePtr SelectedProfile;
+	/** Pointer to active profile. */
+	TSharedPtr<ILauncherProfile> SelectedProfile;
 
 private:
 
-	// Holds a delegate to be invoked when the project path has been modified.
+	/** A delegate to be invoked when the project path has been modified. */
 	FSimpleMulticastDelegate ProjectPathChangedDelegate;
 
-	// Holds a delegate to be invoked when the profile list has been modified.
+	/** A delegate to be invoked when the profile list has been modified. */
 	FSimpleMulticastDelegate ProfileListChangedDelegate;
 
-	// Holds a delegate to be invoked when the selected profile changed.
+	/** A delegate to be invoked when the selected profile changed. */
 	FOnSelectedSProjectLauncherProfileChanged ProfileSelectedDelegate;
-
 };
diff --git a/Engine/Source/Developer/ProjectLauncher/Private/ProjectLauncherHelpers.h b/Engine/Source/Developer/ProjectLauncher/Private/ProjectLauncherHelpers.h
--- a/Engine/Source/Developer/ProjectLauncher/Private/ProjectLauncherHelpers.h
+++ b/Engine/Source/Developer/ProjectLauncher/Private/ProjectLauncherHelpers.h
@@ -2,24 +2,24 @@
 
 #pragma once
 
-#include "CoreMinimal.h"
+#include "EditorStyleSet.h"
 #include "Styling/SlateColor.h"
 #include "Widgets/DeclarativeSyntaxSupport.h"
 #include "Widgets/Images/SImage.h"
 #include "Widgets/Input/SButton.h"
-#include "EditorStyleSet.h"
 
-struct FSlateBrush;
 class FText;
+struct FSlateBrush;
+
 
 /* Static helpers
  *****************************************************************************/
 
 // @todo gmp: move this into a shared library or create a SImageButton widget.
-static TSharedRef<SButton> MakeImageButton( const FSlateBrush* ButtonImage, const FText& ButtonTip, const FOnClicked& OnClickedDelegate )
+static TSharedRef<SButton> MakeImageButton(const FSlateBrush* ButtonImage, const FText& ButtonTip, const FOnClicked& OnClickedDelegate)
 {
 	return SNew(SButton)
-		.ButtonStyle( FEditorStyle::Get(), "NoBorder" )
+		.ButtonStyle(FEditorStyle::Get(), "NoBorder")
 		.OnClicked(OnClickedDelegate)
 		.ToolTipText(ButtonTip)
 		.ContentPadding(2)
diff --git a/Engine/Source/Developer/ProjectLauncher/Private/ProjectLauncherModule.cpp b/Engine/Source/Developer/ProjectLauncher/Private/ProjectLauncherModule.cpp
--- a/Engine/Source/Developer/ProjectLauncher/Private/ProjectLauncherModule.cpp
+++ b/Engine/Source/Developer/ProjectLauncher/Private/ProjectLauncherModule.cpp
@@ -1,22 +1,25 @@
 // Copyright 1998-2017 Epic Games, Inc. All Rights Reserved.
 
-#include "CoreMinimal.h"
+#include "ILauncherServicesModule.h"
+#include "IProjectLauncherModule.h"
+#include "ITargetDeviceServicesModule.h"
 #include "Modules/ModuleManager.h"
-#include "Interfaces/ITargetDeviceServicesModule.h"
-#include "Interfaces/ILauncherServicesModule.h"
-#include "Widgets/DeclarativeSyntaxSupport.h"
-#include "Textures/SlateIcon.h"
-#include "Framework/Docking/TabManager.h"
+#include "WorkspaceMenuStructure.h"
+#include "WorkspaceMenuStructureModule.h"
+
 #include "EditorStyleSet.h"
-#include "Interfaces/IProjectLauncherModule.h"
+#include "Framework/Docking/TabManager.h"
+#include "Textures/SlateIcon.h"
+#include "Widgets/DeclarativeSyntaxSupport.h"
+#include "Widgets/Docking/SDockTab.h"
+
 #include "Models/ProjectLauncherModel.h"
 #include "Widgets/SProjectLauncher.h"
-#include "Widgets/Docking/SDockTab.h"
-#include "WorkspaceMenuStructure.h"
-#include "WorkspaceMenuStructureModule.h"
+
 
 static const FName ProjectLauncherTabName("ProjectLauncher");
 
+
 /**
  * Implements the SessionSProjectLauncher module.
  */
@@ -25,9 +28,9 @@ class FProjectLauncherModule
 {
 public:
 
-	// IModuleInterface interface
+	//~ IModuleInterface interface
 	
-	virtual void StartupModule( ) override
+	virtual void StartupModule() override
 	{
 #if WITH_EDITOR
 		FGlobalTabmanager::Get()->RegisterTabSpawner(ProjectLauncherTabName, FOnSpawnTab::CreateRaw(this, &FProjectLauncherModule::SpawnProjectLauncherTab));
@@ -42,7 +45,7 @@ public:
 #endif
 	}
 
-	virtual void ShutdownModule( ) override
+	virtual void ShutdownModule() override
 	{
 #if WITH_EDITOR
 		FGlobalTabmanager::Get()->UnregisterTabSpawner(ProjectLauncherTabName);
@@ -68,7 +71,7 @@ private:
 		ILauncherServicesModule& ProjectLauncherServicesModule = FModuleManager::LoadModuleChecked<ILauncherServicesModule>("LauncherServices");
 		ITargetDeviceServicesModule& TargetDeviceServicesModule = FModuleManager::LoadModuleChecked<ITargetDeviceServicesModule>("TargetDeviceServices");
 
-		FProjectLauncherModelRef Model = MakeShareable(new FProjectLauncherModel(
+		TSharedRef<FProjectLauncherModel> Model = MakeShareable(new FProjectLauncherModel(
 			TargetDeviceServicesModule.GetDeviceProxyManager(),
 			ProjectLauncherServicesModule.CreateLauncher(),
 			ProjectLauncherServicesModule.GetProfileManager()
diff --git a/Engine/Source/Developer/ProjectLauncher/Private/Widgets/Archive/SProjectLauncherArchivePage.cpp b/Engine/Source/Developer/ProjectLauncher/Private/Widgets/Archive/SProjectLauncherArchivePage.cpp
--- a/Engine/Source/Developer/ProjectLauncher/Private/Widgets/Archive/SProjectLauncherArchivePage.cpp
+++ b/Engine/Source/Developer/ProjectLauncher/Private/Widgets/Archive/SProjectLauncherArchivePage.cpp
@@ -1,16 +1,17 @@
 // Copyright 1998-2017 Epic Games, Inc. All Rights Reserved.
 
-#include "Widgets/Archive/SProjectLauncherArchivePage.h"
-#include "Widgets/SBoxPanel.h"
-#include "Styling/SlateTypes.h"
+#include "SProjectLauncherArchivePage.h"
+
+#include "EditorStyleSet.h"
+#include "DesktopPlatformModule.h"
 #include "Framework/Application/SlateApplication.h"
-#include "Widgets/Layout/SBorder.h"
-#include "Widgets/Text/STextBlock.h"
-#include "Widgets/Input/SEditableTextBox.h"
+#include "Styling/SlateTypes.h"
+#include "Widgets/SBoxPanel.h"
 #include "Widgets/Input/SButton.h"
 #include "Widgets/Input/SCheckBox.h"
-#include "EditorStyleSet.h"
-#include "DesktopPlatformModule.h"
+#include "Widgets/Input/SEditableTextBox.h"
+#include "Widgets/Layout/SBorder.h"
+#include "Widgets/Text/STextBlock.h"
 
 
 #define LOCTEXT_NAMESPACE "SProjectLauncherArchivePage"
@@ -19,14 +20,14 @@
 /* SProjectLauncherCookPage structors
  *****************************************************************************/
 
-SProjectLauncherArchivePage::~SProjectLauncherArchivePage( )
+SProjectLauncherArchivePage::~SProjectLauncherArchivePage()
 {
 }
 
 /* SProjectLauncherCookPage interface
  *****************************************************************************/
 
-void SProjectLauncherArchivePage::Construct( const FArguments& InArgs, const FProjectLauncherModelRef& InModel )
+void SProjectLauncherArchivePage::Construct(const FArguments& InArgs, const TSharedRef<FProjectLauncherModel>& InModel)
 {
 	Model = InModel;
 
@@ -53,8 +54,8 @@ void SProjectLauncherArchivePage::Construct( const FArguments& InArgs, const FPr
 					[
 						// archive mode check box
 						SNew(SCheckBox)
-						.IsChecked(this, &SProjectLauncherArchivePage::HandleArchiveIsChecked)
-						.OnCheckStateChanged(this, &SProjectLauncherArchivePage::HandleArchiveCheckedStateChanged)
+							.IsChecked(this, &SProjectLauncherArchivePage::HandleArchiveIsChecked)
+							.OnCheckStateChanged(this, &SProjectLauncherArchivePage::HandleArchiveCheckedStateChanged)
 					]
 			]
 
@@ -63,58 +64,59 @@ void SProjectLauncherArchivePage::Construct( const FArguments& InArgs, const FPr
 			.Padding(0.0, 8.0, 0.0, 0.0)
 			[
 				SNew(SBorder)
-				.Padding(8.0)
-				.BorderImage(FEditorStyle::GetBrush("ToolPanel.GroupBorder"))
-				.Visibility(this, &SProjectLauncherArchivePage::HandleArchiveVisibility)
-				[
-					SNew(SVerticalBox)
-					
-					+ SVerticalBox::Slot()
-					.AutoHeight()
+					.Padding(8.0)
+					.BorderImage(FEditorStyle::GetBrush("ToolPanel.GroupBorder"))
+					.Visibility(this, &SProjectLauncherArchivePage::HandleArchiveVisibility)
 					[
-						SNew(STextBlock)
-						.Text(LOCTEXT("ArchiveDirectoryTitle", "Archive Directory Path:"))
-					]
+						SNew(SVerticalBox)
 					
-					+ SVerticalBox::Slot()
-					.AutoHeight()
-					.Padding(0.0, 4.0, 0.0, 0.0)
-					[
-						SNew(SHorizontalBox)
-
-						+ SHorizontalBox::Slot()
-						.FillWidth(1.0)
-						.Padding(0.0, 0.0, 0.0, 3.0)
-						[
-							// archive path text box
-							SNew(SEditableTextBox)
-							.Text(this, &SProjectLauncherArchivePage::GetDirectoryPathText)
-							.OnTextCommitted(this, &SProjectLauncherArchivePage::OnDirectoryTextCommitted)
-						]
-
-						+ SHorizontalBox::Slot()
-						.AutoWidth()
-						.HAlign(HAlign_Right)
-						.Padding(4.0, 0.0, 0.0, 0.0)
-						[
-							// browse button
-							SNew(SButton)
-							.ContentPadding(FMargin(6.0, 2.0))
-							.IsEnabled(this, &SProjectLauncherArchivePage::IsDirectoryEditable)
-							.Text(LOCTEXT("BrowseButtonText", "Browse..."))
-							.ToolTipText(LOCTEXT("BrowseButtonToolTip", "Browse for the directory"))
-							.OnClicked(this, &SProjectLauncherArchivePage::HandleBrowseButtonClicked)
-						]
+						+ SVerticalBox::Slot()
+							.AutoHeight()
+							[
+								SNew(STextBlock)
+									.Text(LOCTEXT("ArchiveDirectoryTitle", "Archive Directory Path:"))
+							]
+					
+						+ SVerticalBox::Slot()
+							.AutoHeight()
+							.Padding(0.0, 4.0, 0.0, 0.0)
+							[
+								SNew(SHorizontalBox)
+
+								+ SHorizontalBox::Slot()
+									.FillWidth(1.0)
+									.Padding(0.0, 0.0, 0.0, 3.0)
+									[
+										// archive path text box
+										SNew(SEditableTextBox)
+											.Text(this, &SProjectLauncherArchivePage::GetDirectoryPathText)
+											.OnTextCommitted(this, &SProjectLauncherArchivePage::OnDirectoryTextCommitted)
+									]
+
+								+ SHorizontalBox::Slot()
+									.AutoWidth()
+									.HAlign(HAlign_Right)
+									.Padding(4.0, 0.0, 0.0, 0.0)
+									[
+										// browse button
+										SNew(SButton)
+											.ContentPadding(FMargin(6.0, 2.0))
+											.IsEnabled(this, &SProjectLauncherArchivePage::IsDirectoryEditable)
+											.Text(LOCTEXT("BrowseButtonText", "Browse..."))
+											.ToolTipText(LOCTEXT("BrowseButtonToolTip", "Browse for the directory"))
+											.OnClicked(this, &SProjectLauncherArchivePage::HandleBrowseButtonClicked)
+									]
+							]
 					]
-				]
 			]
 	];
 }
 
+
 /* SProjectLauncherArchivePage callbacks
  *****************************************************************************/
 
-void SProjectLauncherArchivePage::HandleArchiveCheckedStateChanged( ECheckBoxState CheckState )
+void SProjectLauncherArchivePage::HandleArchiveCheckedStateChanged(ECheckBoxState CheckState)
 {
 	ILauncherProfilePtr SelectedProfile = Model->GetSelectedProfile();
 
@@ -140,7 +142,8 @@ ECheckBoxState SProjectLauncherArchivePage::HandleArchiveIsChecked() const
 	return ECheckBoxState::Unchecked;
 }
 
-EVisibility SProjectLauncherArchivePage::HandleArchiveVisibility( ) const
+
+EVisibility SProjectLauncherArchivePage::HandleArchiveVisibility() const
 {
 	ILauncherProfilePtr SelectedProfile = Model->GetSelectedProfile();
 
@@ -155,6 +158,7 @@ EVisibility SProjectLauncherArchivePage::HandleArchiveVisibility( ) const
 	return EVisibility::Collapsed;
 }
 
+
 FText SProjectLauncherArchivePage::GetDirectoryPathText() const
 {
 	ILauncherProfilePtr SelectedProfile = Model->GetSelectedProfile();
@@ -167,6 +171,7 @@ FText SProjectLauncherArchivePage::GetDirectoryPathText() const
 	return FText::GetEmpty();
 }
 
+
 FReply SProjectLauncherArchivePage::HandleBrowseButtonClicked()
 {
 	IDesktopPlatform* DesktopPlatform = FDesktopPlatformModule::Get();
@@ -202,6 +207,7 @@ FReply SProjectLauncherArchivePage::HandleBrowseButtonClicked()
 	return FReply::Handled();
 }
 
+
 bool SProjectLauncherArchivePage::IsDirectoryEditable() const
 {
 	ILauncherProfilePtr SelectedProfile = Model->GetSelectedProfile();
@@ -213,6 +219,7 @@ bool SProjectLauncherArchivePage::IsDirectoryEditable() const
 	return false;
 }
 
+
 void SProjectLauncherArchivePage::OnDirectoryTextCommitted(const FText& InText, ETextCommit::Type CommitInfo)
 {
 	//if (CommitInfo == ETextCommit::OnEnter)
@@ -226,4 +233,5 @@ void SProjectLauncherArchivePage::OnDirectoryTextCommitted(const FText& InText,
 	}
 }
 
+
 #undef LOCTEXT_NAMESPACE
diff --git a/Engine/Source/Developer/ProjectLauncher/Private/Widgets/Archive/SProjectLauncherArchivePage.h b/Engine/Source/Developer/ProjectLauncher/Private/Widgets/Archive/SProjectLauncherArchivePage.h
--- a/Engine/Source/Developer/ProjectLauncher/Private/Widgets/Archive/SProjectLauncherArchivePage.h
+++ b/Engine/Source/Developer/ProjectLauncher/Private/Widgets/Archive/SProjectLauncherArchivePage.h
@@ -35,7 +35,7 @@ public:
 	 * @param InArgs The Slate argument list.
 	 * @param InModel The data model.
 	 */
-	void Construct(	const FArguments& InArgs, const FProjectLauncherModelRef& InModel );
+	void Construct(	const FArguments& InArgs, const TSharedRef<FProjectLauncherModel>& InModel );
 
 private:
 
@@ -63,5 +63,5 @@ private:
 private:
 
 	/** Holds a pointer to the data model. */
-	FProjectLauncherModelPtr Model;
+	TSharedPtr<FProjectLauncherModel> Model;
 };
diff --git a/Engine/Source/Developer/ProjectLauncher/Private/Widgets/Build/SProjectLauncherBuildPage.cpp b/Engine/Source/Developer/ProjectLauncher/Private/Widgets/Build/SProjectLauncherBuildPage.cpp
--- a/Engine/Source/Developer/ProjectLauncher/Private/Widgets/Build/SProjectLauncherBuildPage.cpp
+++ b/Engine/Source/Developer/ProjectLauncher/Private/Widgets/Build/SProjectLauncherBuildPage.cpp
@@ -1,17 +1,19 @@
 // Copyright 1998-2017 Epic Games, Inc. All Rights Reserved.
 
-#include "Widgets/Build/SProjectLauncherBuildPage.h"
-#include "Widgets/SBoxPanel.h"
-#include "Styling/SlateTypes.h"
+#include "SProjectLauncherBuildPage.h"
+
+#include "EditorStyleSet.h"
 #include "Styling/CoreStyle.h"
+#include "Styling/SlateTypes.h"
+#include "Widgets/SBoxPanel.h"
+#include "Widgets/Input/SCheckBox.h"
 #include "Widgets/Layout/SBorder.h"
 #include "Widgets/Text/STextBlock.h"
-#include "Widgets/Input/SCheckBox.h"
-#include "EditorStyleSet.h"
-#include "Widgets/Shared/SProjectLauncherBuildConfigurationSelector.h"
-#include "Widgets/Shared/SProjectLauncherFormLabel.h"
+
 #include "Widgets/Cook/SProjectLauncherCookedPlatforms.h"
 #include "Widgets/Layout/SExpandableArea.h"
+#include "Widgets/Shared/SProjectLauncherBuildConfigurationSelector.h"
+#include "Widgets/Shared/SProjectLauncherFormLabel.h"
 
 
 #define LOCTEXT_NAMESPACE "SProjectLauncherBuildPage"
@@ -20,7 +22,7 @@
 /* SProjectLauncherCookPage structors
  *****************************************************************************/
 
-SProjectLauncherBuildPage::~SProjectLauncherBuildPage( )
+SProjectLauncherBuildPage::~SProjectLauncherBuildPage()
 {
 	if (Model.IsValid())
 	{
@@ -32,7 +34,7 @@ SProjectLauncherBuildPage::~SProjectLauncherBuildPage( )
 /* SProjectLauncherCookPage interface
  *****************************************************************************/
 
-void SProjectLauncherBuildPage::Construct( const FArguments& InArgs, const FProjectLauncherModelRef& InModel )
+void SProjectLauncherBuildPage::Construct(const FArguments& InArgs, const TSharedRef<FProjectLauncherModel>& InModel)
 {
 	Model = InModel;
 
@@ -59,8 +61,8 @@ void SProjectLauncherBuildPage::Construct( const FArguments& InArgs, const FProj
 					[
 						// build mode check box
 						SNew(SCheckBox)
-						.IsChecked(this, &SProjectLauncherBuildPage::HandleBuildIsChecked)
-						.OnCheckStateChanged(this, &SProjectLauncherBuildPage::HandleBuildCheckedStateChanged)
+							.IsChecked(this, &SProjectLauncherBuildPage::HandleBuildIsChecked)
+							.OnCheckStateChanged(this, &SProjectLauncherBuildPage::HandleBuildCheckedStateChanged)
 					]
 			]
 
@@ -69,30 +71,30 @@ void SProjectLauncherBuildPage::Construct( const FArguments& InArgs, const FProj
 			.Padding(0, 3, 0, 3)
 			[
 				SNew(SBorder)
-				.BorderImage(FEditorStyle::GetBrush("ToolPanel.GroupBorder"))
-				.Visibility(this, &SProjectLauncherBuildPage::ShowBuildConfiguration)
-				[
-					SNew(SHorizontalBox)
-
-					+ SHorizontalBox::Slot()
-					.AutoWidth()
+					.BorderImage(FEditorStyle::GetBrush("ToolPanel.GroupBorder"))
+					.Visibility(this, &SProjectLauncherBuildPage::ShowBuildConfiguration)
 					[
-						SNew(SProjectLauncherFormLabel)
-						.ErrorToolTipText(NSLOCTEXT("SProjectLauncherBuildValidation", "NoBuildConfigurationSelectedError", "A Build Configuration must be selected."))
-						.ErrorVisibility(this, &SProjectLauncherBuildPage::HandleValidationErrorIconVisibility, ELauncherProfileValidationErrors::NoBuildConfigurationSelected)
-						.LabelText(LOCTEXT("ConfigurationComboBoxLabel", "Build Configuration:"))
+						SNew(SHorizontalBox)
+
+						+ SHorizontalBox::Slot()
+							.AutoWidth()
+							[
+								SNew(SProjectLauncherFormLabel)
+									.ErrorToolTipText(NSLOCTEXT("SProjectLauncherBuildValidation", "NoBuildConfigurationSelectedError", "A Build Configuration must be selected."))
+									.ErrorVisibility(this, &SProjectLauncherBuildPage::HandleValidationErrorIconVisibility, ELauncherProfileValidationErrors::NoBuildConfigurationSelected)
+									.LabelText(LOCTEXT("ConfigurationComboBoxLabel", "Build Configuration:"))
+							]
+
+						+ SHorizontalBox::Slot()
+							.AutoWidth()
+							[
+								// build configuration selector
+								SNew(SProjectLauncherBuildConfigurationSelector)
+									.Font(FCoreStyle::Get().GetFontStyle(TEXT("NormalFont")))
+									.OnConfigurationSelected(this, &SProjectLauncherBuildPage::HandleBuildConfigurationSelectorConfigurationSelected)
+									.Text(this, &SProjectLauncherBuildPage::HandleBuildConfigurationSelectorText)
+							]
 					]
-
-					+ SHorizontalBox::Slot()
-					.AutoWidth()
-					[
-						// build configuration selector
-						SNew(SProjectLauncherBuildConfigurationSelector)
-						.Font(FCoreStyle::Get().GetFontStyle(TEXT("NormalFont")))
-						.OnConfigurationSelected(this, &SProjectLauncherBuildPage::HandleBuildConfigurationSelectorConfigurationSelected)
-						.Text(this, &SProjectLauncherBuildPage::HandleBuildConfigurationSelectorText)
-					]
-				]
 			]
 
         + SVerticalBox::Slot()
@@ -100,37 +102,37 @@ void SProjectLauncherBuildPage::Construct( const FArguments& InArgs, const FProj
             .Padding(0.0f, 8.0f, 0.0f, 0.0f)
             [
                 SNew(SExpandableArea)
-                .AreaTitle(LOCTEXT("AdvancedAreaTitle", "Advanced Settings"))
-                .InitiallyCollapsed(true)
-                .Padding(8.0)
-                .BodyContent()
-                [
-                    SNew(SVerticalBox)
-
-/*                    + SVerticalBox::Slot()
-                    .AutoHeight()
-                    [
-                        SNew(SButton)
-                        .Text(LOCTEXT("GenDSYMText", "Generate DSYM"))
-                        .IsEnabled( this, &SProjectLauncherBuildPage::HandleGenDSYMButtonEnabled )
-                        .OnClicked( this, &SProjectLauncherBuildPage::HandleGenDSYMClicked )
-                    ]*/
-					+ SVerticalBox::Slot()
-					.AutoHeight()
+					.AreaTitle(LOCTEXT("AdvancedAreaTitle", "Advanced Settings"))
+					.InitiallyCollapsed(true)
+					.Padding(8.0)
+					.BodyContent()
 					[
-						// build mode check box
-						SNew(SCheckBox)
-						.IsChecked(this, &SProjectLauncherBuildPage::HandleUATIsChecked)
-						.OnCheckStateChanged(this, &SProjectLauncherBuildPage::HandleUATCheckedStateChanged)
-						.Padding(FMargin(4.0f, 0.0f))
-						.ToolTipText(LOCTEXT("UATCheckBoxTooltip", "If checked, UAT will be built as part of the build."))
-						.Content()
+						SNew(SVerticalBox)
+
+	/*                    + SVerticalBox::Slot()
+						.AutoHeight()
+						[
+							SNew(SButton)
+							.Text(LOCTEXT("GenDSYMText", "Generate DSYM"))
+							.IsEnabled(this, &SProjectLauncherBuildPage::HandleGenDSYMButtonEnabled)
+							.OnClicked(this, &SProjectLauncherBuildPage::HandleGenDSYMClicked)
+						]*/
+						+ SVerticalBox::Slot()
+						.AutoHeight()
 						[
-							SNew(STextBlock)
-							.Text(LOCTEXT("UATCheckBoxText", "Build UAT"))
+							// build mode check box
+							SNew(SCheckBox)
+								.IsChecked(this, &SProjectLauncherBuildPage::HandleUATIsChecked)
+								.OnCheckStateChanged(this, &SProjectLauncherBuildPage::HandleUATCheckedStateChanged)
+								.Padding(FMargin(4.0f, 0.0f))
+								.ToolTipText(LOCTEXT("UATCheckBoxTooltip", "If checked, UAT will be built as part of the build."))
+								.Content()
+								[
+									SNew(STextBlock)
+										.Text(LOCTEXT("UATCheckBoxText", "Build UAT"))
+								]
 						]
-					]
-               ]
+				   ]
             ]
 
 		+ SVerticalBox::Slot()
@@ -148,7 +150,7 @@ void SProjectLauncherBuildPage::Construct( const FArguments& InArgs, const FProj
 /* SProjectLauncherBuildPage implementation
  *****************************************************************************/
 
-bool SProjectLauncherBuildPage::GenerateDSYMForProject( const FString& ProjectName, const FString& Configuration )
+bool SProjectLauncherBuildPage::GenerateDSYMForProject(const FString& ProjectName, const FString& Configuration)
 {
     // UAT executable
     FString ExecutablePath = FPaths::ConvertRelativePathToFull(FPaths::EngineDir() + FString(TEXT("Build")) / TEXT("BatchFiles"));
@@ -173,6 +175,7 @@ bool SProjectLauncherBuildPage::GenerateDSYMForProject( const FString& ProjectNa
 		FPlatformProcess::CloseProc(ProcessHandle);
 		return true;
 	}
+
 	return false;
 }
 
@@ -180,7 +183,7 @@ bool SProjectLauncherBuildPage::GenerateDSYMForProject( const FString& ProjectNa
 /* SProjectLauncherBuildPage callbacks
  *****************************************************************************/
 
-void SProjectLauncherBuildPage::HandleBuildCheckedStateChanged( ECheckBoxState CheckState )
+void SProjectLauncherBuildPage::HandleBuildCheckedStateChanged(ECheckBoxState CheckState)
 {
 	ILauncherProfilePtr SelectedProfile = Model->GetSelectedProfile();
 
@@ -207,13 +210,13 @@ ECheckBoxState SProjectLauncherBuildPage::HandleBuildIsChecked() const
 }
 
 
-void SProjectLauncherBuildPage::HandleProfileManagerProfileSelected( const ILauncherProfilePtr& SelectedProfile, const ILauncherProfilePtr& PreviousProfile )
+void SProjectLauncherBuildPage::HandleProfileManagerProfileSelected(const ILauncherProfilePtr& SelectedProfile, const ILauncherProfilePtr& PreviousProfile)
 {
 	// reload settings
 }
 
 
-EVisibility SProjectLauncherBuildPage::HandleBuildPlatformVisibility( ) const
+EVisibility SProjectLauncherBuildPage::HandleBuildPlatformVisibility() const
 {
 	ILauncherProfilePtr SelectedProfile = Model->GetSelectedProfile();
 
@@ -240,7 +243,7 @@ FReply SProjectLauncherBuildPage::HandleGenDSYMClicked()
             FString ProjectName = SelectedProfile->GetProjectName();
             EBuildConfigurations::Type ProjectConfig = SelectedProfile->GetBuildConfiguration();
 
-            GenerateDSYMForProject( ProjectName, EBuildConfigurations::ToString(ProjectConfig) );
+            GenerateDSYMForProject(ProjectName, EBuildConfigurations::ToString(ProjectConfig));
         }
     }
 
@@ -318,7 +321,7 @@ EVisibility SProjectLauncherBuildPage::HandleValidationErrorIconVisibility(ELaun
 	return EVisibility::Hidden;
 }
 
-void SProjectLauncherBuildPage::HandleUATCheckedStateChanged( ECheckBoxState CheckState )
+void SProjectLauncherBuildPage::HandleUATCheckedStateChanged(ECheckBoxState CheckState)
 {
 	ILauncherProfilePtr SelectedProfile = Model->GetSelectedProfile();
 
@@ -345,5 +348,4 @@ ECheckBoxState SProjectLauncherBuildPage::HandleUATIsChecked() const
 }
 
 
-
 #undef LOCTEXT_NAMESPACE
diff --git a/Engine/Source/Developer/ProjectLauncher/Private/Widgets/Build/SProjectLauncherBuildPage.h b/Engine/Source/Developer/ProjectLauncher/Private/Widgets/Build/SProjectLauncherBuildPage.h
--- a/Engine/Source/Developer/ProjectLauncher/Private/Widgets/Build/SProjectLauncherBuildPage.h
+++ b/Engine/Source/Developer/ProjectLauncher/Private/Widgets/Build/SProjectLauncherBuildPage.h
@@ -36,7 +36,7 @@ public:
 	 * @param InArgs The Slate argument list.
 	 * @param InModel The data model.
 	 */
-	void Construct(	const FArguments& InArgs, const FProjectLauncherModelRef& InModel );
+	void Construct(	const FArguments& InArgs, const TSharedRef<FProjectLauncherModel>& InModel );
 
 protected:
 
@@ -90,5 +90,5 @@ private:
 private:
 
 	/** Holds a pointer to the data model. */
-	FProjectLauncherModelPtr Model;
+	TSharedPtr<FProjectLauncherModel> Model;
 };
