Copying //UE4/Dev-Sequencer to //UE4/Main (Source: //UE4/Dev-Sequencer @ 3617642)

diff --git a/Engine/Source/Developer/ProjectLauncher/Private/Widgets/Cook/SProjectLauncherCookByTheBookSettings.cpp b/Engine/Source/Developer/ProjectLauncher/Private/Widgets/Cook/SProjectLauncherCookByTheBookSettings.cpp
--- a/Engine/Source/Developer/ProjectLauncher/Private/Widgets/Cook/SProjectLauncherCookByTheBookSettings.cpp
+++ b/Engine/Source/Developer/ProjectLauncher/Private/Widgets/Cook/SProjectLauncherCookByTheBookSettings.cpp
@@ -1,19 +1,21 @@
 // Copyright 1998-2017 Epic Games, Inc. All Rights Reserved.
 
-#include "Widgets/Cook/SProjectLauncherCookByTheBookSettings.h"
-#include "SlateOptMacros.h"
+#include "SProjectLauncherCookByTheBookSettings.h"
+
+#include "DesktopPlatformModule.h"
+#include "EditorStyleSet.h"
 #include "Framework/Application/SlateApplication.h"
-#include "Widgets/Layout/SSeparator.h"
+#include "SlateOptMacros.h"
 #include "Widgets/Images/SImage.h"
-#include "Widgets/Text/STextBlock.h"
-#include "Widgets/Input/SEditableTextBox.h"
 #include "Widgets/Input/SButton.h"
 #include "Widgets/Input/SCheckBox.h"
-#include "EditorStyleSet.h"
-#include "DesktopPlatformModule.h"
+#include "Widgets/Input/SEditableTextBox.h"
+#include "Widgets/Layout/SSeparator.h"
+#include "Widgets/Text/STextBlock.h"
+#include "Widgets/Views/SListView.h"
+
 #include "Widgets/Shared/SProjectLauncherBuildConfigurationSelector.h"
 #include "Widgets/Shared/SProjectLauncherFormLabel.h"
-#include "Widgets/Views/SListView.h"
 #include "Widgets/Cook/SProjectLauncherMapListRow.h"
 #include "Widgets/Cook/SProjectLauncherCultureListRow.h"
 #include "Widgets/Cook/SProjectLauncherCookedPlatforms.h"
@@ -27,7 +29,7 @@
 /* SProjectLauncherCookByTheBookSettings structors
  *****************************************************************************/
 
-SProjectLauncherCookByTheBookSettings::~SProjectLauncherCookByTheBookSettings( )
+SProjectLauncherCookByTheBookSettings::~SProjectLauncherCookByTheBookSettings()
 {
 	if (Model.IsValid())
 	{
@@ -40,7 +42,7 @@ SProjectLauncherCookByTheBookSettings::~SProjectLauncherCookByTheBookSettings( )
  *****************************************************************************/
 
 BEGIN_SLATE_FUNCTION_BUILD_OPTIMIZATION
-void SProjectLauncherCookByTheBookSettings::Construct( const FArguments& InArgs, const FProjectLauncherModelRef& InModel, bool InShowSimple )
+void SProjectLauncherCookByTheBookSettings::Construct(const FArguments& InArgs, const TSharedRef<FProjectLauncherModel>& InModel, bool InShowSimple)
 {
 	Model = InModel;
 
@@ -73,32 +75,32 @@ TSharedRef<SWidget> SProjectLauncherCookByTheBookSettings::MakeComplexWidget()
 {
 	TSharedRef<SWidget> Widget = SNew(SVerticalBox)
 
-		+ SVerticalBox::Slot()
+	+ SVerticalBox::Slot()
 		.AutoHeight()
 		.MaxHeight(256.0f)
 		[
 			SNew(SBorder)
-			.BorderImage(FEditorStyle::GetBrush("ToolPanel.GroupBorder"))
-			.Padding(8.0f)
-			[
-				SNew(SVerticalBox)
-
-				+ SVerticalBox::Slot()
-				.AutoHeight()
+				.BorderImage(FEditorStyle::GetBrush("ToolPanel.GroupBorder"))
+				.Padding(8.0f)
 				[
-					SNew(SProjectLauncherFormLabel)
-					.ErrorToolTipText(NSLOCTEXT("ProjectLauncherBuildValidation", "NoCookedPlatformSelectedError", "At least one Platform must be selected when cooking by the book."))
-					.ErrorVisibility(this, &SProjectLauncherCookByTheBookSettings::HandleValidationErrorIconVisibility, ELauncherProfileValidationErrors::NoPlatformSelected)
-					.LabelText(LOCTEXT("CookedPlatformsLabel", "Cooked Platforms:"))
-				]
+					SNew(SVerticalBox)
 
-				+ SVerticalBox::Slot()
-					.FillHeight(1.0)
-					.Padding(0.0f, 2.0f, 0.0f, 0.0f)
+					+ SVerticalBox::Slot()
+					.AutoHeight()
 					[
-						SNew(SProjectLauncherCookedPlatforms, Model.ToSharedRef())
+						SNew(SProjectLauncherFormLabel)
+							.ErrorToolTipText(NSLOCTEXT("ProjectLauncherBuildValidation", "NoCookedPlatformSelectedError", "At least one Platform must be selected when cooking by the book."))
+							.ErrorVisibility(this, &SProjectLauncherCookByTheBookSettings::HandleValidationErrorIconVisibility, ELauncherProfileValidationErrors::NoPlatformSelected)
+							.LabelText(LOCTEXT("CookedPlatformsLabel", "Cooked Platforms:"))
 					]
-			]
+
+					+ SVerticalBox::Slot()
+						.FillHeight(1.0)
+						.Padding(0.0f, 2.0f, 0.0f, 0.0f)
+						[
+							SNew(SProjectLauncherCookedPlatforms, Model.ToSharedRef())
+						]
+				]
 		]
 
 	+ SVerticalBox::Slot()
@@ -113,13 +115,13 @@ TSharedRef<SWidget> SProjectLauncherCookByTheBookSettings::MakeComplexWidget()
 				SNew(SVerticalBox)
 
 				+ SVerticalBox::Slot()
-				.AutoHeight()
-				[
-					SNew(SProjectLauncherFormLabel)
-					.ErrorToolTipText(NSLOCTEXT("ProjectLauncherBuildValidation", "NoCookedCulturesSelectedError", "At least one Culture must be selected when cooking by the book."))
-					.ErrorVisibility(this, &SProjectLauncherCookByTheBookSettings::HandleValidationErrorIconVisibility, ELauncherProfileValidationErrors::NoCookedCulturesSelected)
-					.LabelText(LOCTEXT("CookedCulturesLabel", "Cooked Cultures:"))
-				]
+					.AutoHeight()
+					[
+						SNew(SProjectLauncherFormLabel)
+							.ErrorToolTipText(NSLOCTEXT("ProjectLauncherBuildValidation", "NoCookedCulturesSelectedError", "At least one Culture must be selected when cooking by the book."))
+							.ErrorVisibility(this, &SProjectLauncherCookByTheBookSettings::HandleValidationErrorIconVisibility, ELauncherProfileValidationErrors::NoCookedCulturesSelected)
+							.LabelText(LOCTEXT("CookedCulturesLabel", "Cooked Cultures:"))
+					]
 
 				+ SVerticalBox::Slot()
 					.FillHeight(1.0)
@@ -127,19 +129,19 @@ TSharedRef<SWidget> SProjectLauncherCookByTheBookSettings::MakeComplexWidget()
 					[
 						// culture menu
 						SAssignNew(CultureListView, SListView<TSharedPtr<FString> >)
-						.HeaderRow
-						(
-						SNew(SHeaderRow)
-						.Visibility(EVisibility::Collapsed)
-
-						+ SHeaderRow::Column("Culture")
-						.DefaultLabel(LOCTEXT("CultureListMapNameColumnHeader", "Culture"))
-						.FillWidth(1.0f)
-						)
-						.ItemHeight(16.0f)
-						.ListItemsSource(&CultureList)
-						.OnGenerateRow(this, &SProjectLauncherCookByTheBookSettings::HandleCultureListViewGenerateRow)
-						.SelectionMode(ESelectionMode::None)
+							.HeaderRow
+							(
+								SNew(SHeaderRow)
+									.Visibility(EVisibility::Collapsed)
+
+								+ SHeaderRow::Column("Culture")
+									.DefaultLabel(LOCTEXT("CultureListMapNameColumnHeader", "Culture"))
+									.FillWidth(1.0f)
+							)
+							.ItemHeight(16.0f)
+							.ListItemsSource(&CultureList)
+							.OnGenerateRow(this, &SProjectLauncherCookByTheBookSettings::HandleCultureListViewGenerateRow)
+							.SelectionMode(ESelectionMode::None)
 					]
 
 				+ SVerticalBox::Slot()
@@ -147,7 +149,7 @@ TSharedRef<SWidget> SProjectLauncherCookByTheBookSettings::MakeComplexWidget()
 					.Padding(0.0f, 6.0f, 0.0f, 4.0f)
 					[
 						SNew(SSeparator)
-						.Orientation(Orient_Horizontal)
+							.Orientation(Orient_Horizontal)
 					]
 
 				+ SVerticalBox::Slot()
@@ -156,12 +158,12 @@ TSharedRef<SWidget> SProjectLauncherCookByTheBookSettings::MakeComplexWidget()
 						SNew(SHorizontalBox)
 
 						+ SHorizontalBox::Slot()
-						.FillWidth(1.0f)
-						.HAlign(HAlign_Right)
-						[
-							SNew(STextBlock)
-							.Text(LOCTEXT("SelectLabel", "Select:"))
-						]
+							.FillWidth(1.0f)
+							.HAlign(HAlign_Right)
+							[
+								SNew(STextBlock)
+									.Text(LOCTEXT("SelectLabel", "Select:"))
+							]
 
 						+ SHorizontalBox::Slot()
 							.AutoWidth()
@@ -169,10 +171,10 @@ TSharedRef<SWidget> SProjectLauncherCookByTheBookSettings::MakeComplexWidget()
 							[
 								// all cultures hyper link
 								SNew(SHyperlink)
-								.OnNavigate(this, &SProjectLauncherCookByTheBookSettings::HandleAllCulturesHyperlinkNavigate, true)
-								.Text(LOCTEXT("AllPlatformsHyperlinkLabel", "All"))
-								.ToolTipText(LOCTEXT("AllPlatformsButtonTooltip", "Select all available platforms."))
-								.Visibility(this, &SProjectLauncherCookByTheBookSettings::HandleAllCulturesHyperlinkVisibility)									
+									.OnNavigate(this, &SProjectLauncherCookByTheBookSettings::HandleAllCulturesHyperlinkNavigate, true)
+									.Text(LOCTEXT("AllPlatformsHyperlinkLabel", "All"))
+									.ToolTipText(LOCTEXT("AllPlatformsButtonTooltip", "Select all available platforms."))
+									.Visibility(this, &SProjectLauncherCookByTheBookSettings::HandleAllCulturesHyperlinkVisibility)
 							]
 
 						+ SHorizontalBox::Slot()
@@ -180,10 +182,10 @@ TSharedRef<SWidget> SProjectLauncherCookByTheBookSettings::MakeComplexWidget()
 							[
 								// no cultures hyper link
 								SNew(SHyperlink)
-								.OnNavigate(this, &SProjectLauncherCookByTheBookSettings::HandleAllCulturesHyperlinkNavigate, false)
-								.Text(LOCTEXT("NoCulturesHyperlinkLabel", "None"))
-								.ToolTipText(LOCTEXT("NoCulturesHyperlinkTooltip", "Deselect all platforms."))
-								.Visibility(this, &SProjectLauncherCookByTheBookSettings::HandleAllCulturesHyperlinkVisibility)									
+									.OnNavigate(this, &SProjectLauncherCookByTheBookSettings::HandleAllCulturesHyperlinkNavigate, false)
+									.Text(LOCTEXT("NoCulturesHyperlinkLabel", "None"))
+									.ToolTipText(LOCTEXT("NoCulturesHyperlinkTooltip", "Deselect all platforms."))
+									.Visibility(this, &SProjectLauncherCookByTheBookSettings::HandleAllCulturesHyperlinkVisibility)
 							]
 					]
 			]
@@ -195,850 +197,859 @@ TSharedRef<SWidget> SProjectLauncherCookByTheBookSettings::MakeComplexWidget()
 		.Padding(0.0f, 8.0f, 0.0f, 0.0f)
 		[
 			SNew(SBorder)
-			.BorderImage(FEditorStyle::GetBrush("ToolPanel.GroupBorder"))
-			.Padding(8.0f)
-			[
-				SNew(SVerticalBox)
-
-				+ SVerticalBox::Slot()
-				.AutoHeight()
+				.BorderImage(FEditorStyle::GetBrush("ToolPanel.GroupBorder"))
+				.Padding(8.0f)
 				[
-					SNew(SProjectLauncherFormLabel)
-					.LabelText(LOCTEXT("CookedMapsLabel", "Cooked Maps:"))
-				]
-
-				+ SVerticalBox::Slot()
-					.AutoHeight()
-					.Padding(0.0f, 4.0f)
-					.VAlign(VAlign_Center)
-					[
-						SNew(SHorizontalBox)
-
-						+ SHorizontalBox::Slot()
-						.AutoWidth()
-						[
-							// all maps radio button
-							SNew(SCheckBox)
-							.IsChecked(this, &SProjectLauncherCookByTheBookSettings::HandleShowCheckBoxIsChecked, EShowMapsChoices::ShowAllMaps)
-							.OnCheckStateChanged(this, &SProjectLauncherCookByTheBookSettings::HandleShowCheckBoxCheckStateChanged, EShowMapsChoices::ShowAllMaps)
-							.Style(FEditorStyle::Get(), "RadioButton")
-							[
-								SNew(STextBlock)
-								.Text(LOCTEXT("AllMapsCheckBoxText", "Show all"))
-							]
-						]
-
-						+ SHorizontalBox::Slot()
-							.FillWidth(1.0f)
-							.Padding(8.0f, 0.0f, 0.0f, 0.0f)
-							[
-								// cooked maps radio button
-								SNew(SCheckBox)
-								.IsChecked(this, &SProjectLauncherCookByTheBookSettings::HandleShowCheckBoxIsChecked, EShowMapsChoices::ShowCookedMaps)
-								.OnCheckStateChanged(this, &SProjectLauncherCookByTheBookSettings::HandleShowCheckBoxCheckStateChanged, EShowMapsChoices::ShowCookedMaps)
-								.Style(FEditorStyle::Get(), "RadioButton")
-								[
-									SNew(STextBlock)
-									.Text(LOCTEXT("CookedMapsCheckBoxText", "Show cooked"))
-								]
-							]
-					]
-
-				+ SVerticalBox::Slot()
-					.FillHeight(1.0)
-					.Padding(0.0f, 4.0f, 0.0f, 0.0f)
-					[
-						// map list
-						SAssignNew(MapListView, SListView<TSharedPtr<FString> >)
-						.HeaderRow
-						(
-						SNew(SHeaderRow)
-						.Visibility(EVisibility::Collapsed)
-
-						+ SHeaderRow::Column("MapName")
-						.DefaultLabel(LOCTEXT("MapListMapNameColumnHeader", "Map"))
-						.FillWidth(1.0f)
-						)
-						.ItemHeight(16.0f)
-						.ListItemsSource(&MapList)
-						.OnGenerateRow(this, &SProjectLauncherCookByTheBookSettings::HandleMapListViewGenerateRow)
-						.SelectionMode(ESelectionMode::None)
-					]
-
-				+ SVerticalBox::Slot()
-					.AutoHeight()
-					[
-						SNew(SHorizontalBox)
-						.Visibility(this, &SProjectLauncherCookByTheBookSettings::HandleNoMapSelectedBoxVisibility)
+					SNew(SVerticalBox)
 
-						+ SHorizontalBox::Slot()
-						.AutoWidth()
+					+ SVerticalBox::Slot()
+						.AutoHeight()
 						[
-							SNew(SImage)
-							.Image(FEditorStyle::GetBrush(TEXT("Icons.Warning")))
+							SNew(SProjectLauncherFormLabel)
+								.LabelText(LOCTEXT("CookedMapsLabel", "Cooked Maps:"))
 						]
 
-						+ SHorizontalBox::Slot()
-							.AutoWidth()
-							.Padding(4.0f, 0.0f)
-							.VAlign(VAlign_Center)
-							[
-								SNew(STextBlock)
-								.Text(this, &SProjectLauncherCookByTheBookSettings::HandleNoMapsTextBlockText)
-							]
-					]
-
-				+ SVerticalBox::Slot()
-					.AutoHeight()
-					.Padding(0.0f, 6.0f, 0.0f, 4.0f)
-					[
-						SNew(SSeparator)
-						.Orientation(Orient_Horizontal)
-					]
-
-				+ SVerticalBox::Slot()
-					.AutoHeight()
-					.VAlign(VAlign_Center)
-					[
-						SNew(SHorizontalBox)
-
-						+ SHorizontalBox::Slot()
-						.FillWidth(1.0f)
-						.HAlign(HAlign_Right)
+					+ SVerticalBox::Slot()
+						.AutoHeight()
+						.Padding(0.0f, 4.0f)
+						.VAlign(VAlign_Center)
 						[
-							SNew(STextBlock)
-							.Text(LOCTEXT("SelectLabel", "Select:"))
-							.Visibility(this, &SProjectLauncherCookByTheBookSettings::HandleMapSelectionHyperlinkVisibility)
-						]
-
-						+ SHorizontalBox::Slot()
-							.AutoWidth()
-							.Padding(8.0f, 0.0f)
-							[
-								// all maps hyper link
-								SNew(SHyperlink)
-								.OnNavigate(this, &SProjectLauncherCookByTheBookSettings::HandleAllMapsHyperlinkNavigate, true)
-								.Text(LOCTEXT("AllMapsHyperlinkLabel", "All"))
-								.ToolTipText(LOCTEXT("AllMapsHyperlinkTooltip", "Select all available maps."))
-								.Visibility(this, &SProjectLauncherCookByTheBookSettings::HandleMapSelectionHyperlinkVisibility)									
-							]
+							SNew(SHorizontalBox)
 
-						+ SHorizontalBox::Slot()
+							+ SHorizontalBox::Slot()
 							.AutoWidth()
 							[
-								// no maps hyper link
-								SNew(SHyperlink)
-								.OnNavigate(this, &SProjectLauncherCookByTheBookSettings::HandleAllMapsHyperlinkNavigate, false)
-								.Text(LOCTEXT("NoMapsHyperlinkLabel", "None"))
-								.ToolTipText(LOCTEXT("NoMapsHyperlinkTooltip", "Deselect all maps."))
-								.Visibility(this, &SProjectLauncherCookByTheBookSettings::HandleMapSelectionHyperlinkVisibility)
+								// all maps radio button
+								SNew(SCheckBox)
+									.IsChecked(this, &SProjectLauncherCookByTheBookSettings::HandleShowCheckBoxIsChecked, EShowMapsChoices::ShowAllMaps)
+									.OnCheckStateChanged(this, &SProjectLauncherCookByTheBookSettings::HandleShowCheckBoxCheckStateChanged, EShowMapsChoices::ShowAllMaps)
+									.Style(FEditorStyle::Get(), "RadioButton")
+									[
+										SNew(STextBlock)
+											.Text(LOCTEXT("AllMapsCheckBoxText", "Show all"))
+									]
 							]
-					]
-			]
-		]
-
-
-
-
-
-		+ SVerticalBox::Slot()
-			.AutoHeight()
-			.Padding(0.0f, 8.0f, 0.0f, 0.0f)
-			[
-				SNew(SExpandableArea)
-				.AreaTitle(LOCTEXT("PatchingAreaTitle", "Release / DLC / Patching Settings"))
-				.InitiallyCollapsed(true)
-				.Padding(8.0f)
-				.BodyContent()
-				[
-					SNew(SVerticalBox)
-
 
-					//////////////////////////////////////////////////////////////////////////
-					// create release version options
-					+ SVerticalBox::Slot()
-					.AutoHeight()
-					.Padding(0.0f, 4.0f, 0.0f, 0.0f)
-					[
-						// unreal pak check box
-						SNew(SCheckBox)
-						.IsChecked(this, &SProjectLauncherCookByTheBookSettings::HandleCreateReleaseVersionCheckBoxIsChecked)
-						.OnCheckStateChanged(this, &SProjectLauncherCookByTheBookSettings::HandleCreateReleaseVersionCheckBoxCheckStateChanged)
-						.Padding(FMargin(4.0f, 0.0f))
-						.ToolTipText(LOCTEXT("CreateReleaseVersionCheckBoxTooltip", "Create a release version of the game for distribution."))
-						.Content()
-						[
-							SNew(STextBlock)
-							.Text(LOCTEXT("CreateReleaseVersionBoxText", "Create a release version of the game for distribution."))
+							+ SHorizontalBox::Slot()
+								.FillWidth(1.0f)
+								.Padding(8.0f, 0.0f, 0.0f, 0.0f)
+								[
+									// cooked maps radio button
+									SNew(SCheckBox)
+										.IsChecked(this, &SProjectLauncherCookByTheBookSettings::HandleShowCheckBoxIsChecked, EShowMapsChoices::ShowCookedMaps)
+										.OnCheckStateChanged(this, &SProjectLauncherCookByTheBookSettings::HandleShowCheckBoxCheckStateChanged, EShowMapsChoices::ShowCookedMaps)
+										.Style(FEditorStyle::Get(), "RadioButton")
+										[
+											SNew(STextBlock)
+												.Text(LOCTEXT("CookedMapsCheckBoxText", "Show cooked"))
+										]
+								]
 						]
-					]
 
 					+ SVerticalBox::Slot()
-						.AutoHeight()
-						.Padding(0.0f, 8.0f, 0.0f, 0.0f)
+						.FillHeight(1.0)
+						.Padding(0.0f, 4.0f, 0.0f, 0.0f)
 						[
-							SNew(SProjectLauncherFormLabel)
-							.LabelText(LOCTEXT("CreateReleaseVersionTextBoxLabel", "Name of the new release to create."))
+							// map list
+							SAssignNew(MapListView, SListView<TSharedPtr<FString> >)
+								.HeaderRow
+								(
+									SNew(SHeaderRow)
+										.Visibility(EVisibility::Collapsed)
+
+									+ SHeaderRow::Column("MapName")
+										.DefaultLabel(LOCTEXT("MapListMapNameColumnHeader", "Map"))
+										.FillWidth(1.0f)
+								)
+								.ItemHeight(16.0f)
+								.ListItemsSource(&MapList)
+								.OnGenerateRow(this, &SProjectLauncherCookByTheBookSettings::HandleMapListViewGenerateRow)
+								.SelectionMode(ESelectionMode::None)
 						]
 
 					+ SVerticalBox::Slot()
 						.AutoHeight()
-						.Padding(0.0, 4.0, 0.0, 0.0)
 						[
 							SNew(SHorizontalBox)
+								.Visibility(this, &SProjectLauncherCookByTheBookSettings::HandleNoMapSelectedBoxVisibility)
 
 							+ SHorizontalBox::Slot()
-							.FillWidth(1.0)
-							.Padding(0.0, 0.0, 0.0, 3.0)
-							[
-								// repository path text box
-								SNew(SEditableTextBox)
-								.ToolTipText(LOCTEXT("CreateReleaseVersionTextBoxTooltip", "Name of the new release to create."))
-								.Text(this, &SProjectLauncherCookByTheBookSettings::HandleCreateReleaseVersionNameTextBlockText)
-								.OnTextCommitted(this, &SProjectLauncherCookByTheBookSettings::HandleCreateReleaseVersionNameCommitted)
-							]
+								.AutoWidth()
+								[
+									SNew(SImage)
+										.Image(FEditorStyle::GetBrush(TEXT("Icons.Warning")))
+								]
+
+							+ SHorizontalBox::Slot()
+								.AutoWidth()
+								.Padding(4.0f, 0.0f)
+								.VAlign(VAlign_Center)
+								[
+									SNew(STextBlock)
+										.Text(this, &SProjectLauncherCookByTheBookSettings::HandleNoMapsTextBlockText)
+								]
 						]
 
-					+SVerticalBox::Slot()
+					+ SVerticalBox::Slot()
 						.AutoHeight()
-						.Padding(0.0f, 8.0f, 0.0f, 0.0f)
+						.Padding(0.0f, 6.0f, 0.0f, 4.0f)
 						[
-							SNew(SProjectLauncherFormLabel)
-							.ToolTipText(LOCTEXT("BasedOnReleaseVersionTextBoxToolTip", "The release version which this DLC / Patch / Next release is based on."))
-						.LabelText(LOCTEXT("BasedOnReleaseVersionTextBoxLabel", "Release version this is based on."))
+							SNew(SSeparator)
+								.Orientation(Orient_Horizontal)
 						]
 
 					+ SVerticalBox::Slot()
 						.AutoHeight()
-						.Padding(0.0, 4.0, 0.0, 0.0)
+						.VAlign(VAlign_Center)
 						[
 							SNew(SHorizontalBox)
 
 							+ SHorizontalBox::Slot()
-							.FillWidth(1.0)
-							.Padding(0.0, 0.0, 0.0, 3.0)
-							[
-								// repository path text box
-								SNew(SEditableTextBox)
-								.ToolTipText(LOCTEXT("NextReleaseVersionTextBoxTooltip", "Release version to base the next release / DLC / patch on."))
-								.Text(this, &SProjectLauncherCookByTheBookSettings::HandleBasedOnReleaseVersionNameTextBlockText)
-								.OnTextCommitted(this, &SProjectLauncherCookByTheBookSettings::HandleBasedOnReleaseVersionNameCommitted)
-							]
+								.FillWidth(1.0f)
+								.HAlign(HAlign_Right)
+								[
+									SNew(STextBlock)
+										.Text(LOCTEXT("SelectLabel", "Select:"))
+										.Visibility(this, &SProjectLauncherCookByTheBookSettings::HandleMapSelectionHyperlinkVisibility)
+								]
+
+							+ SHorizontalBox::Slot()
+								.AutoWidth()
+								.Padding(8.0f, 0.0f)
+								[
+									// all maps hyper link
+									SNew(SHyperlink)
+										.OnNavigate(this, &SProjectLauncherCookByTheBookSettings::HandleAllMapsHyperlinkNavigate, true)
+										.Text(LOCTEXT("AllMapsHyperlinkLabel", "All"))
+										.ToolTipText(LOCTEXT("AllMapsHyperlinkTooltip", "Select all available maps."))
+										.Visibility(this, &SProjectLauncherCookByTheBookSettings::HandleMapSelectionHyperlinkVisibility)
+								]
+
+							+ SHorizontalBox::Slot()
+								.AutoWidth()
+								[
+									// no maps hyper link
+									SNew(SHyperlink)
+										.OnNavigate(this, &SProjectLauncherCookByTheBookSettings::HandleAllMapsHyperlinkNavigate, false)
+										.Text(LOCTEXT("NoMapsHyperlinkLabel", "None"))
+										.ToolTipText(LOCTEXT("NoMapsHyperlinkTooltip", "Deselect all maps."))
+										.Visibility(this, &SProjectLauncherCookByTheBookSettings::HandleMapSelectionHyperlinkVisibility)
+								]
 						]
+				]
+		]
 
-					+ SVerticalBox::Slot()
-						.AutoHeight()
-						.Padding(0.0f, 4.0f, 0.0f, 0.0f)
-						[
-							// stage base release pak files check box
-							SNew(SCheckBox)
-							.IsChecked(this, &SProjectLauncherCookByTheBookSettings::HandleStageBaseReleasePaksCheckBoxIsChecked)
-							.OnCheckStateChanged(this, &SProjectLauncherCookByTheBookSettings::HandleStageBaseReleasePaksCheckBoxCheckStateChanged)
-							.Padding(FMargin(4.0f, 0.0f))
-							.ToolTipText(LOCTEXT("StageBaseReleasePaksCheckBoxTooltip", "If checked, unchanged pak files present in the base release version will be staged."))
-							.Content()
+		+ SVerticalBox::Slot()
+			.AutoHeight()
+			.Padding(0.0f, 8.0f, 0.0f, 0.0f)
+			[
+				SNew(SExpandableArea)
+					.AreaTitle(LOCTEXT("PatchingAreaTitle", "Release / DLC / Patching Settings"))
+					.InitiallyCollapsed(true)
+					.Padding(8.0f)
+					.BodyContent()
+					[
+						SNew(SVerticalBox)
+
+						//////////////////////////////////////////////////////////////////////////
+						// create release version options
+						+ SVerticalBox::Slot()
+							.AutoHeight()
+							.Padding(0.0f, 4.0f, 0.0f, 0.0f)
 							[
-								SNew(STextBlock)
-								.Text(LOCTEXT("StageBaseReleasePaksCheckBoxText", "Stage base release pak files"))
+								// unreal pak check box
+								SNew(SCheckBox)
+									.IsChecked(this, &SProjectLauncherCookByTheBookSettings::HandleCreateReleaseVersionCheckBoxIsChecked)
+									.OnCheckStateChanged(this, &SProjectLauncherCookByTheBookSettings::HandleCreateReleaseVersionCheckBoxCheckStateChanged)
+									.Padding(FMargin(4.0f, 0.0f))
+									.ToolTipText(LOCTEXT("CreateReleaseVersionCheckBoxTooltip", "Create a release version of the game for distribution."))
+									.Content()
+									[
+										SNew(STextBlock)
+											.Text(LOCTEXT("CreateReleaseVersionBoxText", "Create a release version of the game for distribution."))
+									]
 							]
-						]
-
 
-					// end create release version
-					//////////////////////////////////////////////////////////////////////////
-					// generate patch params
-					+ SVerticalBox::Slot()
-						.AutoHeight()
-						.Padding(0.0f, 4.0f, 0.0f, 0.0f)
-						[
-							// unreal pak check box
-							SNew(SCheckBox)
-							.IsChecked(this, &SProjectLauncherCookByTheBookSettings::HandleGeneratePatchCheckBoxIsChecked)
-							.OnCheckStateChanged(this, &SProjectLauncherCookByTheBookSettings::HandleGeneratePatchCheckBoxCheckStateChanged)
-							.Padding(FMargin(4.0f, 0.0f))
-							.ToolTipText(LOCTEXT("GeneratePatchCheckBoxTooltip", "If checked, content will be diffed against source content and only changed files will be included in new pak"))
-							.Content()
+						+ SVerticalBox::Slot()
+							.AutoHeight()
+							.Padding(0.0f, 8.0f, 0.0f, 0.0f)
 							[
-								SNew(STextBlock)
-								.Text(LOCTEXT("GeneratePatchCheckBoxText", "Generate patch"))
+								SNew(SProjectLauncherFormLabel)
+									.LabelText(LOCTEXT("CreateReleaseVersionTextBoxLabel", "Name of the new release to create."))
 							]
-						]
 
-					// end generate patch options
-					//////////////////////////////////////////////////////////////////////////
-					// generate new patch level params
-					+ SVerticalBox::Slot()
-						.AutoHeight()
-						.Padding(0.0f, 4.0f, 0.0f, 0.0f)
-						[
-							SNew(SCheckBox)
-							.IsChecked(this, &SProjectLauncherCookByTheBookSettings::HandleAddPatchLevelCheckBoxIsChecked)
-							.OnCheckStateChanged(this, &SProjectLauncherCookByTheBookSettings::HandleAddPatchLevelCheckBoxCheckStateChanged)
-							.Padding(FMargin(4.0f, 0.0f))
-							.ToolTipText(LOCTEXT("AddPatchLevelCheckBoxTooltip", "If checked, a new numbered pak will be generated with patch content"))
-							.Content()
+						+ SVerticalBox::Slot()
+							.AutoHeight()
+							.Padding(0.0, 4.0, 0.0, 0.0)
 							[
-								SNew(STextBlock)
-								.Text(LOCTEXT("AddPatchLevelCheckBoxText", "Add a new patch tier"))
+								SNew(SHorizontalBox)
+
+								+ SHorizontalBox::Slot()
+								.FillWidth(1.0)
+								.Padding(0.0, 0.0, 0.0, 3.0)
+								[
+									// repository path text box
+									SNew(SEditableTextBox)
+										.ToolTipText(LOCTEXT("CreateReleaseVersionTextBoxTooltip", "Name of the new release to create."))
+										.Text(this, &SProjectLauncherCookByTheBookSettings::HandleCreateReleaseVersionNameTextBlockText)
+										.OnTextCommitted(this, &SProjectLauncherCookByTheBookSettings::HandleCreateReleaseVersionNameCommitted)
+								]
 							]
-						]
 
-					// end generate new patch level options
-					//////////////////////////////////////////////////////////////////////////
-					// generate dlc options
-					+ SVerticalBox::Slot()
-						.AutoHeight()
-						.Padding(0.0f, 4.0f, 0.0f, 0.0f)
-						[
-							// unreal pak check box
-							SNew(SCheckBox)
-							.IsChecked(this, &SProjectLauncherCookByTheBookSettings::HandleBuildDLCCheckBoxIsChecked)
-							.OnCheckStateChanged(this, &SProjectLauncherCookByTheBookSettings::HandleBuildDLCCheckBoxCheckStateChanged)
-							.Padding(FMargin(4.0f, 0.0f))
-							.ToolTipText(LOCTEXT("BuildDLCCheckBoxTooltip", "If checked, DLC will be built without the content released with the original game."))
-							.Content()
+						+SVerticalBox::Slot()
+							.AutoHeight()
+							.Padding(0.0f, 8.0f, 0.0f, 0.0f)
 							[
-								SNew(STextBlock)
-								.Text(LOCTEXT("BuildDLCCheckBoxText", "Build DLC"))
+								SNew(SProjectLauncherFormLabel)
+									.ToolTipText(LOCTEXT("BasedOnReleaseVersionTextBoxToolTip", "The release version which this DLC / Patch / Next release is based on."))
+									.LabelText(LOCTEXT("BasedOnReleaseVersionTextBoxLabel", "Release version this is based on."))
 							]
-						]
 
+						+ SVerticalBox::Slot()
+							.AutoHeight()
+							.Padding(0.0, 4.0, 0.0, 0.0)
+							[
+								SNew(SHorizontalBox)
 
-					+ SVerticalBox::Slot()
-						.AutoHeight()
-						.Padding(0.0f, 8.0f, 0.0f, 0.0f)
-						[
-							SNew(SProjectLauncherFormLabel)
-							.LabelText(LOCTEXT("DLCNameTextBoxLabel", "Name of the DLC to build."))
-						]
+								+ SHorizontalBox::Slot()
+									.FillWidth(1.0)
+									.Padding(0.0, 0.0, 0.0, 3.0)
+									[
+										// repository path text box
+										SNew(SEditableTextBox)
+											.ToolTipText(LOCTEXT("NextReleaseVersionTextBoxTooltip", "Release version to base the next release / DLC / patch on."))
+											.Text(this, &SProjectLauncherCookByTheBookSettings::HandleBasedOnReleaseVersionNameTextBlockText)
+											.OnTextCommitted(this, &SProjectLauncherCookByTheBookSettings::HandleBasedOnReleaseVersionNameCommitted)
+									]
+							]
 
-					+ SVerticalBox::Slot()
-						.AutoHeight()
-						.Padding(0.0, 4.0, 0.0, 0.0)
-						[
-							SNew(SHorizontalBox)
+						// end create release version
+						//////////////////////////////////////////////////////////////////////////
 
-							+ SHorizontalBox::Slot()
-							.FillWidth(1.0)
-							.Padding(0.0, 0.0, 0.0, 3.0)
+						// generate patch params
+						+ SVerticalBox::Slot()
+							.AutoHeight()
+							.Padding(0.0f, 4.0f, 0.0f, 0.0f)
 							[
-								// repository path text box
-								SNew(SEditableTextBox)
-								.ToolTipText(LOCTEXT("DLCNameTextBoxTooltip", "Name of DLC to build."))
-								.Text(this, &SProjectLauncherCookByTheBookSettings::HandleDLCNameTextBlockText)
-								.OnTextCommitted(this, &SProjectLauncherCookByTheBookSettings::HandleDLCNameCommitted)
+								// unreal pak check box
+								SNew(SCheckBox)
+									.IsChecked(this, &SProjectLauncherCookByTheBookSettings::HandleGeneratePatchCheckBoxIsChecked)
+									.OnCheckStateChanged(this, &SProjectLauncherCookByTheBookSettings::HandleGeneratePatchCheckBoxCheckStateChanged)
+									.Padding(FMargin(4.0f, 0.0f))
+									.ToolTipText(LOCTEXT("GeneratePatchCheckBoxTooltip", "If checked, content will be diffed against source content and only changed files will be included in new pak"))
+									.Content()
+									[
+										SNew(STextBlock)
+											.Text(LOCTEXT("GeneratePatchCheckBoxText", "Generate patch"))
+									]
 							]
-						]
 
-					+ SVerticalBox::Slot()
-						.AutoHeight()
-						.Padding(0.0f, 4.0f, 0.0f, 0.0f)
-						[
-							// unreal pak check box
-							SNew(SCheckBox)
-							.IsChecked(this, &SProjectLauncherCookByTheBookSettings::HandleDLCIncludeEngineContentCheckBoxIsChecked)
-							.OnCheckStateChanged(this, &SProjectLauncherCookByTheBookSettings::HandleDLCIncludeEngineContentCheckBoxCheckStateChanged)
-							.Padding(FMargin(4.0f, 0.0f))
-							.ToolTipText(LOCTEXT("HandleDLCIncludeEngineContentCheckBoxTooltip", "If checked, DLC will include engine content which was not included in original release, if not checked will error when accessing content from engine directory."))
-							.Content()
+						// end generate patch options
+						//////////////////////////////////////////////////////////////////////////
+
+						// generate dlc options
+						+ SVerticalBox::Slot()
+							.AutoHeight()
+							.Padding(0.0f, 4.0f, 0.0f, 0.0f)
 							[
-								SNew(STextBlock)
-								.Text(LOCTEXT("HandleDLCIncludeEngineContentCheckBoxText", "Include engine content"))
+								// unreal pak check box
+								SNew(SCheckBox)
+									.IsChecked(this, &SProjectLauncherCookByTheBookSettings::HandleBuildDLCCheckBoxIsChecked)
+									.OnCheckStateChanged(this, &SProjectLauncherCookByTheBookSettings::HandleBuildDLCCheckBoxCheckStateChanged)
+									.Padding(FMargin(4.0f, 0.0f))
+									.ToolTipText(LOCTEXT("BuildDLCCheckBoxTooltip", "If checked, DLC will be built without the content released with the original game."))
+									.Content()
+									[
+										SNew(STextBlock)
+											.Text(LOCTEXT("BuildDLCCheckBoxText", "Build DLC"))
+									]
+							]
+
+						+ SVerticalBox::Slot()
+							.AutoHeight()
+							.Padding(0.0f, 8.0f, 0.0f, 0.0f)
+							[
+								SNew(SProjectLauncherFormLabel)
+									.LabelText(LOCTEXT("DLCNameTextBoxLabel", "Name of the DLC to build."))
 							]
-						]
-					// end generate dlc 
-					//////////////////////////////////////////////////////////////////////////
-			]
-		]
 
-		
+						+ SVerticalBox::Slot()
+							.AutoHeight()
+							.Padding(0.0, 4.0, 0.0, 0.0)
+							[
+								SNew(SHorizontalBox)
+
+								+ SHorizontalBox::Slot()
+									.FillWidth(1.0)
+									.Padding(0.0, 0.0, 0.0, 3.0)
+									[
+										// repository path text box
+										SNew(SEditableTextBox)
+											.ToolTipText(LOCTEXT("DLCNameTextBoxTooltip", "Name of DLC to build."))
+											.Text(this, &SProjectLauncherCookByTheBookSettings::HandleDLCNameTextBlockText)
+											.OnTextCommitted(this, &SProjectLauncherCookByTheBookSettings::HandleDLCNameCommitted)
+									]
+							]
 
+						+ SVerticalBox::Slot()
+							.AutoHeight()
+							.Padding(0.0f, 4.0f, 0.0f, 0.0f)
+							[
+								// unreal pak check box
+								SNew(SCheckBox)
+									.IsChecked(this, &SProjectLauncherCookByTheBookSettings::HandleDLCIncludeEngineContentCheckBoxIsChecked)
+									.OnCheckStateChanged(this, &SProjectLauncherCookByTheBookSettings::HandleDLCIncludeEngineContentCheckBoxCheckStateChanged)
+									.Padding(FMargin(4.0f, 0.0f))
+									.ToolTipText(LOCTEXT("HandleDLCIncludeEngineContentCheckBoxTooltip", "If checked, DLC will include engine content which was not included in original release, if not checked will error when accessing content from engine directory."))
+									.Content()
+									[
+										SNew(STextBlock)
+											.Text(LOCTEXT("HandleDLCIncludeEngineContentCheckBoxText", "Include engine content"))
+									]
+							]
 
+						// end generate dlc 
+						//////////////////////////////////////////////////////////////////////////
+				]
+		]
 
 	+ SVerticalBox::Slot()
 		.AutoHeight()
 		.Padding(0.0f, 8.0f, 0.0f, 0.0f)
 		[
 			SNew(SExpandableArea)
-			.AreaTitle(LOCTEXT("AdvancedAreaTitle", "Advanced Settings"))
-			.InitiallyCollapsed(true)
-			.Padding(8.0f)
-			.BodyContent()
-			[
-				SNew(SVerticalBox)
+				.AreaTitle(LOCTEXT("AdvancedAreaTitle", "Advanced Settings"))
+				.InitiallyCollapsed(true)
+				.Padding(8.0f)
+				.BodyContent()
+				[
+					SNew(SVerticalBox)
 
-				+ SVerticalBox::Slot()
-					.AutoHeight()
-					[
-						// incremental cook check box
-						SNew(SCheckBox)
-						.IsChecked(this, &SProjectLauncherCookByTheBookSettings::HandleIncrementalCheckBoxIsChecked)
-						.OnCheckStateChanged(this, &SProjectLauncherCookByTheBookSettings::HandleIncrementalCheckBoxCheckStateChanged)
-						.Padding(FMargin(4.0f, 0.0f))
-						.ToolTipText(LOCTEXT("IncrementalCheckBoxTooltip", "If checked, only modified content will be cooked, resulting in much faster cooking times. It is recommended to enable this option whenever possible."))
-						.Content()
+					+ SVerticalBox::Slot()
+						.AutoHeight()
 						[
-							SNew(STextBlock)
-							.Text(LOCTEXT("IncrementalCheckBoxText", "Iterative cooking: Only cook content modified from previous cook"))
+							// incremental cook check box
+							SNew(SCheckBox)
+								.IsChecked(this, &SProjectLauncherCookByTheBookSettings::HandleIncrementalCheckBoxIsChecked)
+								.OnCheckStateChanged(this, &SProjectLauncherCookByTheBookSettings::HandleIncrementalCheckBoxCheckStateChanged)
+								.Padding(FMargin(4.0f, 0.0f))
+								.ToolTipText(LOCTEXT("IncrementalCheckBoxTooltip", "If checked, only modified content will be cooked, resulting in much faster cooking times. It is recommended to enable this option whenever possible."))
+								.Content()
+								[
+									SNew(STextBlock)
+										.Text(LOCTEXT("IncrementalCheckBoxText", "Iterative cooking: Only cook content modified from previous cook"))
+								]
 						]
-					]
-				// disabled for now until this system is live
-				/*+ SVerticalBox::Slot()
-					.AutoHeight()
-					[
-						// incremental cook check box
-						SNew(SCheckBox)
-						.IsChecked(this, &SProjectLauncherCookByTheBookSettings::HandleSharedCookedBuildCheckBoxIsChecked)
-						.OnCheckStateChanged(this, &SProjectLauncherCookByTheBookSettings::HandleSharedCookedBuildCheckBoxCheckStateChanged)
-						.Padding(FMargin(4.0f, 0.0f))
-						.ToolTipText(LOCTEXT("SharedCookedBuildCheckBoxToolTip", "Experimental: Use a build from the network to cook from."))
-						.Content()
+
+					// disabled for now until this system is live
+					/*+ SVerticalBox::Slot()
+						.AutoHeight()
 						[
-							SNew(STextBlock)
-							.Text(LOCTEXT("SharedCookedBuildCheckBoxText", "Iteratively cook from a pre packaged build located on the network"))
-						]
-					]*/
+							// incremental cook check box
+							SNew(SCheckBox)
+							.IsChecked(this, &SProjectLauncherCookByTheBookSettings::HandleSharedCookedBuildCheckBoxIsChecked)
+							.OnCheckStateChanged(this, &SProjectLauncherCookByTheBookSettings::HandleSharedCookedBuildCheckBoxCheckStateChanged)
+							.Padding(FMargin(4.0f, 0.0f))
+							.ToolTipText(LOCTEXT("SharedCookedBuildCheckBoxToolTip", "Experimental: Use a build from the network to cook from."))
+							.Content()
+							[
+								SNew(STextBlock)
+								.Text(LOCTEXT("SharedCookedBuildCheckBoxText", "Iteratively cook from a pre packaged build located on the network"))
+							]
+						]*/
 
-				+ SVerticalBox::Slot()
-					.Padding(0.0f, 4.0f, 0.0f, 0.0f)
-					.AutoHeight()
-					[
-						SNew(SCheckBox)
-						.IsChecked(this, &SProjectLauncherCookByTheBookSettings::HandleCompressedCheckBoxIsChecked)
-						.OnCheckStateChanged(this, &SProjectLauncherCookByTheBookSettings::HandleCompressedCheckBoxCheckStateChanged)
-						.Padding(FMargin(4.0f, 0.0f))
-						.ToolTipText(LOCTEXT("CompressedCheckboxToolTip", "If checked, content will be generated compressed.  These will be smaller but potentially take longer to load"))
-						.Content()
+					+ SVerticalBox::Slot()
+						.Padding(0.0f, 4.0f, 0.0f, 0.0f)
+						.AutoHeight()
 						[
-							SNew(STextBlock)
-							.Text(LOCTEXT("CompressedCheckBoxText", "Compress content"))
+							SNew(SHorizontalBox)
+
+							+ SHorizontalBox::Slot()
+								.FillWidth(1.0)
+								.Padding(0.0, 0.0, 0.0, 3.0)
+								[
+									// repository path text box
+									SNew(SEditableTextBox)
+										.ToolTipText(LOCTEXT("NextReleaseVersionTextBoxTooltip", "Release version to base the next release / DLC / patch on."))
+										.Text(this, &SProjectLauncherCookByTheBookSettings::HandleBasedOnReleaseVersionNameTextBlockText)
+										.OnTextCommitted(this, &SProjectLauncherCookByTheBookSettings::HandleBasedOnReleaseVersionNameCommitted)
+								]
 						]
-					]
 
-				+ SVerticalBox::Slot()
-					.AutoHeight()
-					.Padding(0.0f, 4.0f, 0.0f, 0.0f)
-					[
-						// incremental cook check box
-						SNew(SCheckBox)
-						.IsChecked(this, &SProjectLauncherCookByTheBookSettings::HandleUnversionedCheckBoxIsChecked)
-						.OnCheckStateChanged(this, &SProjectLauncherCookByTheBookSettings::HandleUnversionedCheckBoxCheckStateChanged)
-						.Padding(FMargin(4.0f, 0.0f))
-						.ToolTipText(LOCTEXT("UnversionedCheckBoxTooltip", "If checked, the version is assumed to be current at load. This is potentially dangerous, but results in smaller patch sizes."))
-						.Content()
+					+ SVerticalBox::Slot()
+						.AutoHeight()
+						.Padding(0.0f, 4.0f, 0.0f, 0.0f)
 						[
-							SNew(STextBlock)
-							.Text(LOCTEXT("UnversionedCheckBoxText", "Save packages without versions"))
+							// stage base release pak files check box
+							SNew(SCheckBox)
+								.IsChecked(this, &SProjectLauncherCookByTheBookSettings::HandleStageBaseReleasePaksCheckBoxIsChecked)
+								.OnCheckStateChanged(this, &SProjectLauncherCookByTheBookSettings::HandleStageBaseReleasePaksCheckBoxCheckStateChanged)
+								.Padding(FMargin(4.0f, 0.0f))
+								.ToolTipText(LOCTEXT("StageBaseReleasePaksCheckBoxTooltip", "If checked, unchanged pak files present in the base release version will be staged."))
+								.Content()
+								[
+									SNew(STextBlock)
+										.Text(LOCTEXT("StageBaseReleasePaksCheckBoxText", "Stage base release pak files"))
+								]
 						]
-					]
-
-				// multiprocess cooking options
-				+ SVerticalBox::Slot()
-					.AutoHeight()
-					.Padding(0.0f, 8.0f, 0.0f, 0.0f)
-					[
-						SNew(SProjectLauncherFormLabel)
-						.LabelText(LOCTEXT("MultiProcessCookerTextBoxLabel", "Num cookers to spawn:"))
-					]
-
-				+ SVerticalBox::Slot()
-					.AutoHeight()
-					.Padding(0.0f, 4.0f, 0.0f, 0.0f)
-					[
-						// cooker command line options
-						SNew(SEditableTextBox)
-						.ToolTipText(LOCTEXT("MultiProcessCookerTextBoxTooltip", "The number of cookers to spawn when we do a cook by the book."))
-						.Text(this, &SProjectLauncherCookByTheBookSettings::HandleMultiProcessCookerTextBlockText)
-						.OnTextCommitted(this, &SProjectLauncherCookByTheBookSettings::HandleMultiProcessCookerCommitted)
-					]
 
-				// unreal pak check box
-				+ SVerticalBox::Slot()
-					.AutoHeight()
-					.Padding(0.0f, 4.0f, 0.0f, 0.0f)
-					[
-						// unreal pak check box
-						SNew(SCheckBox)
-						.IsChecked(this, &SProjectLauncherCookByTheBookSettings::HandleUnrealPakCheckBoxIsChecked)
-						.OnCheckStateChanged(this, &SProjectLauncherCookByTheBookSettings::HandleUnrealPakCheckBoxCheckStateChanged)
-						.Padding(FMargin(4.0f, 0.0f))
-						.ToolTipText(LOCTEXT("UnrealPakCheckBoxTooltip", "If checked, the content will be deployed as a single UnrealPak file instead of many separate files."))
-						.Content()
+					// generate patch params
+					+ SVerticalBox::Slot()
+						.AutoHeight()
+						.Padding(0.0f, 4.0f, 0.0f, 0.0f)
 						[
-							SNew(STextBlock)
-							.Text(LOCTEXT("UnrealPakCheckBoxText", "Store all content in a single file (UnrealPak)"))
+							// unreal pak check box
+							SNew(SCheckBox)
+								.IsChecked(this, &SProjectLauncherCookByTheBookSettings::HandleCompressedCheckBoxIsChecked)
+								.OnCheckStateChanged(this, &SProjectLauncherCookByTheBookSettings::HandleCompressedCheckBoxCheckStateChanged)
+								.Padding(FMargin(4.0f, 0.0f))
+								.ToolTipText(LOCTEXT("CompressedCheckboxToolTip", "If checked, content will be generated compressed.  These will be smaller but potentially take longer to load"))
+								.Content()
+								[
+									SNew(STextBlock)
+										.Text(LOCTEXT("CompressedCheckBoxText", "Compress content"))
+								]
 						]
-					]
 
+					// generate new patch level params
+					+ SVerticalBox::Slot()
+						.AutoHeight()
+						.Padding(0.0f, 4.0f, 0.0f, 0.0f)
+						[
+							SNew(SCheckBox)
+								.IsChecked(this, &SProjectLauncherCookByTheBookSettings::HandleAddPatchLevelCheckBoxIsChecked)
+								.OnCheckStateChanged(this, &SProjectLauncherCookByTheBookSettings::HandleAddPatchLevelCheckBoxCheckStateChanged)
+								.Padding(FMargin(4.0f, 0.0f))
+								.ToolTipText(LOCTEXT("AddPatchLevelCheckBoxTooltip", "If checked, a new numbered pak will be generated with patch content"))
+								.Content()
+								[
+									SNew(STextBlock)
+									.Text(LOCTEXT("AddPatchLevelCheckBoxText", "Add a new patch tier"))
+								]
+						]
 
-				+ SVerticalBox::Slot()
-					.Padding(0.0f, 4.0f, 0.0f, 0.0f)
-					.AutoHeight()
-					[
-						SNew(SCheckBox)
-						.IsChecked(this, &SProjectLauncherCookByTheBookSettings::HandleEncryptIniFilesCheckBoxIsChecked)
-					.OnCheckStateChanged(this, &SProjectLauncherCookByTheBookSettings::HandleEncryptIniFilesCheckBoxCheckStateChanged)
-					.Padding(FMargin(4.0f, 0.0f))
-					.ToolTipText(LOCTEXT("EncryptIniFilesCheckboxToolTip", "If checked, ini files stored inside pak file will be encrypted."))
-					.Content()
-					[
-						SNew(STextBlock)
-						.Text(LOCTEXT("EncryptIniFilesCheckBoxText", "Encrypt ini files (only with use pak file)"))
-					]
-					]
-
-				// generate chunks check box
-				+ SVerticalBox::Slot()
-					.AutoHeight()
-					.Padding(0.0f, 4.0f, 0.0f, 0.0f)
-					[
-						
-						SNew(SCheckBox)
-						.IsChecked(this, &SProjectLauncherCookByTheBookSettings::HandleGenerateChunksCheckBoxIsChecked)
-						.OnCheckStateChanged(this, &SProjectLauncherCookByTheBookSettings::HandleGenerateChunksCheckBoxCheckStateChanged)
-						.Padding(FMargin(4.0f, 0.0f))
-						.ToolTipText(LOCTEXT("GenerateChunksCheckBoxTooltip", "If checked, the content will be deployed as multiple UnrealPak files instead of many separate files."))
-						.Content()
+					// generate dlc options
+					+ SVerticalBox::Slot()
+						.AutoHeight()
+						.Padding(0.0f, 4.0f, 0.0f, 0.0f)
 						[
-							SNew(STextBlock)
-							.Text(LOCTEXT("GenerateChunksCheckBoxText", "Generate Chunks"))
+							// incremental cook check box
+							SNew(SCheckBox)
+								.IsChecked(this, &SProjectLauncherCookByTheBookSettings::HandleUnversionedCheckBoxIsChecked)
+								.OnCheckStateChanged(this, &SProjectLauncherCookByTheBookSettings::HandleUnversionedCheckBoxCheckStateChanged)
+								.Padding(FMargin(4.0f, 0.0f))
+								.ToolTipText(LOCTEXT("UnversionedCheckBoxTooltip", "If checked, the version is assumed to be current at load. This is potentially dangerous, but results in smaller patch sizes."))
+								.Content()
+								[
+									SNew(STextBlock)
+										.Text(LOCTEXT("UnversionedCheckBoxText", "Save packages without versions"))
+								]
 						]
-					]
 
-				// don't include editor content
-				+ SVerticalBox::Slot()
-					.AutoHeight()
-					.Padding(0.0f, 4.0f, 0.0f, 0.0f)
-					[
+					// multiprocess cooking options
+					+ SVerticalBox::Slot()
+						.AutoHeight()
+						.Padding(0.0f, 8.0f, 0.0f, 0.0f)
+						[
+							SNew(SProjectLauncherFormLabel)
+								.LabelText(LOCTEXT("MultiProcessCookerTextBoxLabel", "Num cookers to spawn:"))
+						]
 
-						SNew(SCheckBox)
-						.IsChecked(this, &SProjectLauncherCookByTheBookSettings::HandleDontIncludeEditorContentCheckBoxIsChecked)
-						.OnCheckStateChanged(this, &SProjectLauncherCookByTheBookSettings::HandleDontIncludeEditorContentCheckBoxCheckStateChanged)
-						.Padding(FMargin(4.0f, 0.0f))
-						.ToolTipText(LOCTEXT("DontIncludeEditorContentCheckBoxTooltip", "If checked the cooker will skip editor content and not include it in the build."))
-						.Content()
+					+ SVerticalBox::Slot()
+						.AutoHeight()
+						.Padding(0.0f, 4.0f, 0.0f, 0.0f)
 						[
-							SNew(STextBlock)
-							.Text(LOCTEXT("DontIncludeEditorContentCheckBoxText", "Don't Include editor content in the build"))
+							// cooker command line options
+							SNew(SEditableTextBox)
+								.ToolTipText(LOCTEXT("MultiProcessCookerTextBoxTooltip", "The number of cookers to spawn when we do a cook by the book."))
+								.Text(this, &SProjectLauncherCookByTheBookSettings::HandleMultiProcessCookerTextBlockText)
+								.OnTextCommitted(this, &SProjectLauncherCookByTheBookSettings::HandleMultiProcessCookerCommitted)
 						]
-					]
 
+					// unreal pak check box
+					+ SVerticalBox::Slot()
+						.AutoHeight()
+						.Padding(0.0f, 4.0f, 0.0f, 0.0f)
+						[
+							// unreal pak check box
+							SNew(SCheckBox)
+								.IsChecked(this, &SProjectLauncherCookByTheBookSettings::HandleUnrealPakCheckBoxIsChecked)
+								.OnCheckStateChanged(this, &SProjectLauncherCookByTheBookSettings::HandleUnrealPakCheckBoxCheckStateChanged)
+								.Padding(FMargin(4.0f, 0.0f))
+								.ToolTipText(LOCTEXT("UnrealPakCheckBoxTooltip", "If checked, the content will be deployed as a single UnrealPak file instead of many separate files."))
+								.Content()
+								[
+									SNew(STextBlock)
+										.Text(LOCTEXT("UnrealPakCheckBoxText", "Store all content in a single file (UnrealPak)"))
+								]
+						]
 
-				+ SVerticalBox::Slot()
-					.AutoHeight()
-					.Padding(0.0f, 4.0f, 0.0f, 0.0f)
-					[
-						SNew(SExpandableArea)
-						.AreaTitle(LOCTEXT("HttpChunkInstallSettingsAreaTitle", "Http Chunk Install Settings"))
-						.InitiallyCollapsed(true)
-						.Padding(FMargin(4.0f, 0.0f))
-						.BodyContent()
+					+ SVerticalBox::Slot()
+						.Padding(0.0f, 4.0f, 0.0f, 0.0f)
+						.AutoHeight()
 						[
-							SNew(SVerticalBox)
+							SNew(SCheckBox)
+								.IsChecked(this, &SProjectLauncherCookByTheBookSettings::HandleEncryptIniFilesCheckBoxIsChecked)
+								.OnCheckStateChanged(this, &SProjectLauncherCookByTheBookSettings::HandleEncryptIniFilesCheckBoxCheckStateChanged)
+								.Padding(FMargin(4.0f, 0.0f))
+								.ToolTipText(LOCTEXT("EncryptIniFilesCheckboxToolTip", "If checked, ini files stored inside pak file will be encrypted."))
+								.Content()
+								[
+									SNew(STextBlock)
+										.Text(LOCTEXT("EncryptIniFilesCheckBoxText", "Encrypt ini files (only with use pak file)"))
+								]
+						]
 
-							+ SVerticalBox::Slot()
-								.AutoHeight()
-								.Padding(0.0f, 4.0f, 0.0f, 0.0f)
+					// generate chunks check box
+					+ SVerticalBox::Slot()
+						.AutoHeight()
+						.Padding(0.0f, 4.0f, 0.0f, 0.0f)
+						[
+						
+							SNew(SCheckBox)
+								.IsChecked(this, &SProjectLauncherCookByTheBookSettings::HandleGenerateChunksCheckBoxIsChecked)
+								.OnCheckStateChanged(this, &SProjectLauncherCookByTheBookSettings::HandleGenerateChunksCheckBoxCheckStateChanged)
+								.Padding(FMargin(4.0f, 0.0f))
+								.ToolTipText(LOCTEXT("GenerateChunksCheckBoxTooltip", "If checked, the content will be deployed as multiple UnrealPak files instead of many separate files."))
+								.Content()
 								[
-									// unreal pak check box
-									SNew(SCheckBox)
-									.IsChecked(this, &SProjectLauncherCookByTheBookSettings::HandleHttpChunkInstallCheckBoxIsChecked)
-									.OnCheckStateChanged(this, &SProjectLauncherCookByTheBookSettings::HandleHttpChunkInstallCheckBoxCheckStateChanged)
-									.Padding(FMargin(4.0f, 0.0f))
-									.ToolTipText(LOCTEXT("HttpChunkInstallCheckBoxTooltip", "If checked, the content will be split into multiple paks and stored as data that can be downloaded."))
-									.Content()
-									[
-										SNew(STextBlock)
-										.Text(LOCTEXT("HttpChunkInstallCheckBoxText", "Create Http Chunk Install data"))
-									]
+									SNew(STextBlock)
+										.Text(LOCTEXT("GenerateChunksCheckBoxText", "Generate Chunks"))
 								]
+						]
+
+					// don't include editor content
+					+ SVerticalBox::Slot()
+						.AutoHeight()
+						.Padding(0.0f, 4.0f, 0.0f, 0.0f)
+						[
 
-							+ SVerticalBox::Slot()
-								.AutoHeight()
+							SNew(SCheckBox)
+								.IsChecked(this, &SProjectLauncherCookByTheBookSettings::HandleDontIncludeEditorContentCheckBoxIsChecked)
+								.OnCheckStateChanged(this, &SProjectLauncherCookByTheBookSettings::HandleDontIncludeEditorContentCheckBoxCheckStateChanged)
+								.Padding(FMargin(4.0f, 0.0f))
+								.ToolTipText(LOCTEXT("DontIncludeEditorContentCheckBoxTooltip", "If checked the cooker will skip editor content and not include it in the build."))
+								.Content()
 								[
 									SNew(STextBlock)
-									.Text(LOCTEXT("HttpChunkInstallDataPathLabel", "Http Chunk Install Data Path:"))
+										.Text(LOCTEXT("DontIncludeEditorContentCheckBoxText", "Don't Include editor content in the build"))
 								]
+						]
 
-							+ SVerticalBox::Slot()
-								.AutoHeight()
-								.Padding(0.0, 4.0, 0.0, 0.0)
+					+ SVerticalBox::Slot()
+						.AutoHeight()
+						.Padding(0.0f, 4.0f, 0.0f, 0.0f)
+						[
+							SNew(SExpandableArea)
+								.AreaTitle(LOCTEXT("HttpChunkInstallSettingsAreaTitle", "Http Chunk Install Settings"))
+								.InitiallyCollapsed(true)
+								.Padding(FMargin(4.0f, 0.0f))
+								.BodyContent()
 								[
-									SNew(SHorizontalBox)
+									SNew(SVerticalBox)
 
-									+ SHorizontalBox::Slot()
-									.FillWidth(1.0)
-									.Padding(0.0f, 4.0f, 0.0f, 0.0f)
-									[
-										// repository path text box
-										SAssignNew(HttpChunkInstallDirectoryTextBox, SEditableTextBox)
-										.Text(this, &SProjectLauncherCookByTheBookSettings::HandleHtppChunkInstallDirectoryText)
-										.OnTextCommitted(this, &SProjectLauncherCookByTheBookSettings::HandleHtppChunkInstallDirectoryTextCommitted)
-										.OnTextChanged(this, &SProjectLauncherCookByTheBookSettings::HandleHtppChunkInstallDirectoryTextChanged)
-									]
+									+ SVerticalBox::Slot()
+										.AutoHeight()
+										.Padding(0.0f, 4.0f, 0.0f, 0.0f)
+										[
+											// unreal pak check box
+											SNew(SCheckBox)
+												.IsChecked(this, &SProjectLauncherCookByTheBookSettings::HandleHttpChunkInstallCheckBoxIsChecked)
+												.OnCheckStateChanged(this, &SProjectLauncherCookByTheBookSettings::HandleHttpChunkInstallCheckBoxCheckStateChanged)
+												.Padding(FMargin(4.0f, 0.0f))
+												.ToolTipText(LOCTEXT("HttpChunkInstallCheckBoxTooltip", "If checked, the content will be split into multiple paks and stored as data that can be downloaded."))
+												.Content()
+												[
+													SNew(STextBlock)
+														.Text(LOCTEXT("HttpChunkInstallCheckBoxText", "Create Http Chunk Install data"))
+												]
+										]
 
-									+ SHorizontalBox::Slot()
-										.AutoWidth()
-										.HAlign(HAlign_Right)
-										.Padding(4.0, 0.0, 0.0, 0.0)
+									+ SVerticalBox::Slot()
+										.AutoHeight()
 										[
-											// browse button
-											SNew(SButton)
-											.ContentPadding(FMargin(6.0, 2.0))
-											.IsEnabled(true)
-											.Text(LOCTEXT("BrowseButtonText", "Browse..."))
-											.ToolTipText(LOCTEXT("BrowseButtonToolTip", "Browse for the Http Chunk Install Data directory"))
-											.OnClicked(this, &SProjectLauncherCookByTheBookSettings::HandleHtppChunkInstallBrowseButtonClicked)
+											SNew(STextBlock)
+												.Text(LOCTEXT("HttpChunkInstallDataPathLabel", "Http Chunk Install Data Path:"))
+										]
+
+									+ SVerticalBox::Slot()
+										.AutoHeight()
+										.Padding(0.0, 4.0, 0.0, 0.0)
+										[
+											SNew(SHorizontalBox)
+
+											+ SHorizontalBox::Slot()
+												.FillWidth(1.0)
+												.Padding(0.0f, 4.0f, 0.0f, 0.0f)
+												[
+													// repository path text box
+													SAssignNew(HttpChunkInstallDirectoryTextBox, SEditableTextBox)
+														.Text(this, &SProjectLauncherCookByTheBookSettings::HandleHtppChunkInstallDirectoryText)
+														.OnTextCommitted(this, &SProjectLauncherCookByTheBookSettings::HandleHtppChunkInstallDirectoryTextCommitted)
+														.OnTextChanged(this, &SProjectLauncherCookByTheBookSettings::HandleHtppChunkInstallDirectoryTextChanged)
+												]
+
+											+ SHorizontalBox::Slot()
+												.AutoWidth()
+												.HAlign(HAlign_Right)
+												.Padding(4.0, 0.0, 0.0, 0.0)
+												[
+													// browse button
+													SNew(SButton)
+														.ContentPadding(FMargin(6.0, 2.0))
+														.IsEnabled(true)
+														.Text(LOCTEXT("BrowseButtonText", "Browse..."))
+														.ToolTipText(LOCTEXT("BrowseButtonToolTip", "Browse for the Http Chunk Install Data directory"))
+														.OnClicked(this, &SProjectLauncherCookByTheBookSettings::HandleHtppChunkInstallBrowseButtonClicked)
+												]
 										]
-								]
 				
-							+ SVerticalBox::Slot()
-								.AutoHeight()
-								.Padding(0.0f, 4.0f, 0.0f, 0.0f)
-								[
-									SNew(SProjectLauncherFormLabel)
-									.LabelText(LOCTEXT("HttpChunkInstallReleaseTextBoxLabel", "Http Chunk Install Release Name:"))
-								]
+									+ SVerticalBox::Slot()
+										.AutoHeight()
+										.Padding(0.0f, 4.0f, 0.0f, 0.0f)
+										[
+											SNew(SProjectLauncherFormLabel)
+												.LabelText(LOCTEXT("HttpChunkInstallReleaseTextBoxLabel", "Http Chunk Install Release Name:"))
+										]
 
-							+ SVerticalBox::Slot()
-								.AutoHeight()
-								.Padding(0.0f, 4.0f, 0.0f, 0.0f)
-								[
-									// cooker command line options
-									SNew(SEditableTextBox)
-									.ToolTipText(LOCTEXT("HttpChunkInstallReleaseTextBoxTooltip", "Name of this version of the Http Chunk Install data."))
-									.Text(this, &SProjectLauncherCookByTheBookSettings::HandleHttpChunkInstallNameTextBlockText)
-									.OnTextCommitted(this, &SProjectLauncherCookByTheBookSettings::HandleHtppChunkInstallNameCommitted)
+									+ SVerticalBox::Slot()
+										.AutoHeight()
+										.Padding(0.0f, 4.0f, 0.0f, 0.0f)
+										[
+											// cooker command line options
+											SNew(SEditableTextBox)
+												.ToolTipText(LOCTEXT("HttpChunkInstallReleaseTextBoxTooltip", "Name of this version of the Http Chunk Install data."))
+												.Text(this, &SProjectLauncherCookByTheBookSettings::HandleHttpChunkInstallNameTextBlockText)
+												.OnTextCommitted(this, &SProjectLauncherCookByTheBookSettings::HandleHtppChunkInstallNameCommitted)
+										]
 								]
 						]
-					]
-				+ SVerticalBox::Slot()
-					.AutoHeight()
-					.Padding(0.0f, 12.0f, 0.0f, 0.0f)
-					[
-						SNew(SProjectLauncherFormLabel)
-						.LabelText(LOCTEXT("CookConfigurationSelectorLabel", "Cooker build configuration:"))							
-					]
-
-				+ SVerticalBox::Slot()
-					.AutoHeight()
-					.Padding(0.0f, 4.0f, 0.0f, 0.0f)
-					[
-						// cooker build configuration selector
-						SNew(SProjectLauncherBuildConfigurationSelector)
-						.OnConfigurationSelected(this, &SProjectLauncherCookByTheBookSettings::HandleCookConfigurationSelectorConfigurationSelected)
-						.Text(this, &SProjectLauncherCookByTheBookSettings::HandleCookConfigurationSelectorText)
-						.ToolTipText(LOCTEXT("CookConfigurationToolTipText", "Sets the build configuration to use for the cooker commandlet."))
-					]
 
+					+ SVerticalBox::Slot()
+						.AutoHeight()
+						.Padding(0.0f, 12.0f, 0.0f, 0.0f)
+						[
+							SNew(SProjectLauncherFormLabel)
+								.LabelText(LOCTEXT("CookConfigurationSelectorLabel", "Cooker build configuration:"))							
+						]
 
-				// additional cooker options text box
-				+ SVerticalBox::Slot()
-					.AutoHeight()
-					.Padding(0.0f, 8.0f, 0.0f, 0.0f)
-					[
-						SNew(SProjectLauncherFormLabel)
-						.LabelText(LOCTEXT("CookerOptionsTextBoxLabel", "Additional Cooker Options:"))
-					]
-
-				+ SVerticalBox::Slot()
-					.AutoHeight()
-					.Padding(0.0f, 4.0f, 0.0f, 0.0f)
-					[
-						// cooker command line options
-						SNew(SEditableTextBox)
-						.ToolTipText(LOCTEXT("CookerOptionsTextBoxTooltip", "Additional cooker command line parameters can be specified here."))
-						.Text(this, &SProjectLauncherCookByTheBookSettings::HandleCookOptionsTextBlockText)
-						.OnTextCommitted(this, &SProjectLauncherCookByTheBookSettings::HandleCookerOptionsCommitted)
-					]
+					+ SVerticalBox::Slot()
+						.AutoHeight()
+						.Padding(0.0f, 4.0f, 0.0f, 0.0f)
+						[
+							// cooker build configuration selector
+							SNew(SProjectLauncherBuildConfigurationSelector)
+								.OnConfigurationSelected(this, &SProjectLauncherCookByTheBookSettings::HandleCookConfigurationSelectorConfigurationSelected)
+								.Text(this, &SProjectLauncherCookByTheBookSettings::HandleCookConfigurationSelectorText)
+								.ToolTipText(LOCTEXT("CookConfigurationToolTipText", "Sets the build configuration to use for the cooker commandlet."))
+						]
 
+					// additional cooker options text box
+					+ SVerticalBox::Slot()
+						.AutoHeight()
+						.Padding(0.0f, 8.0f, 0.0f, 0.0f)
+						[
+							SNew(SProjectLauncherFormLabel)
+								.LabelText(LOCTEXT("CookerOptionsTextBoxLabel", "Additional Cooker Options:"))
+						]
 
-			]
+					+ SVerticalBox::Slot()
+						.AutoHeight()
+						.Padding(0.0f, 4.0f, 0.0f, 0.0f)
+						[
+							// cooker command line options
+							SNew(SEditableTextBox)
+								.ToolTipText(LOCTEXT("CookerOptionsTextBoxTooltip", "Additional cooker command line parameters can be specified here."))
+								.Text(this, &SProjectLauncherCookByTheBookSettings::HandleCookOptionsTextBlockText)
+								.OnTextCommitted(this, &SProjectLauncherCookByTheBookSettings::HandleCookerOptionsCommitted)
+						]
+				]
 		];
 
 	return Widget;
 }
 END_SLATE_FUNCTION_BUILD_OPTIMIZATION
 
+
 BEGIN_SLATE_FUNCTION_BUILD_OPTIMIZATION
 TSharedRef<SWidget> SProjectLauncherCookByTheBookSettings::MakeSimpleWidget()
 {
 	TSharedRef<SWidget> Widget = SNew(SVerticalBox)
 
-		+ SVerticalBox::Slot()
+	+ SVerticalBox::Slot()
 		.AutoHeight()
 		.MaxHeight(256.0f)
 		[
 			SNew(SBorder)
-			.BorderImage(FEditorStyle::GetBrush("ToolPanel.GroupBorder"))
-			.Padding(8.0f)
-			[
-				SNew(SVerticalBox)
-
-				+ SVerticalBox::Slot()
-				.AutoHeight()
+				.BorderImage(FEditorStyle::GetBrush("ToolPanel.GroupBorder"))
+				.Padding(8.0f)
 				[
-					SNew(SProjectLauncherFormLabel)
-					.ErrorToolTipText(NSLOCTEXT("ProjectLauncherBuildValidation", "NoCookedPlatformSelectedError", "At least one Platform must be selected when cooking by the book."))
-					.ErrorVisibility(this, &SProjectLauncherCookByTheBookSettings::HandleValidationErrorIconVisibility, ELauncherProfileValidationErrors::NoPlatformSelected)
-					.LabelText(LOCTEXT("CookedPlatformsLabel", "Cooked Platforms:"))
-				]
+					SNew(SVerticalBox)
 
-				+ SVerticalBox::Slot()
+					+ SVerticalBox::Slot()
 					.AutoHeight()
-					.Padding(0.0f, 2.0f, 0.0f, 0.0f)
-					[
-						SNew(SProjectLauncherCookedPlatforms, Model.ToSharedRef())
-					]
-			]
+						[
+							SNew(SProjectLauncherFormLabel)
+								.ErrorToolTipText(NSLOCTEXT("ProjectLauncherBuildValidation", "NoCookedPlatformSelectedError", "At least one Platform must be selected when cooking by the book."))
+								.ErrorVisibility(this, &SProjectLauncherCookByTheBookSettings::HandleValidationErrorIconVisibility, ELauncherProfileValidationErrors::NoPlatformSelected)
+								.LabelText(LOCTEXT("CookedPlatformsLabel", "Cooked Platforms:"))
+						]
+
+					+ SVerticalBox::Slot()
+						.AutoHeight()
+						.Padding(0.0f, 2.0f, 0.0f, 0.0f)
+						[
+							SNew(SProjectLauncherCookedPlatforms, Model.ToSharedRef())
+						]
+				]
 		]
 
-		+ SVerticalBox::Slot()
+	+ SVerticalBox::Slot()
 		.AutoHeight()
 		.MaxHeight(256.0f)
 		.Padding(0.0f, 8.0f, 0.0f, 0.0f)
 		[
 			SNew(SBorder)
-			.BorderImage(FEditorStyle::GetBrush("ToolPanel.GroupBorder"))
-			.Padding(8.0f)
-			[
-				SNew(SVerticalBox)
-
-				+ SVerticalBox::Slot()
-				.AutoHeight()
+				.BorderImage(FEditorStyle::GetBrush("ToolPanel.GroupBorder"))
+				.Padding(8.0f)
 				[
-					SNew(SProjectLauncherFormLabel)
-					.LabelText(LOCTEXT("CookedMapsLabel", "Cooked Maps:"))
-				]
-
-				+ SVerticalBox::Slot()
-					.AutoHeight()
-					.Padding(0.0f, 4.0f)
-					.VAlign(VAlign_Center)
-					[
-						SNew(SHorizontalBox)
+					SNew(SVerticalBox)
 
-						+ SHorizontalBox::Slot()
-						.AutoWidth()
+					+ SVerticalBox::Slot()
+						.AutoHeight()
 						[
-							// all maps radio button
-							SNew(SCheckBox)
-							.IsChecked(this, &SProjectLauncherCookByTheBookSettings::HandleShowCheckBoxIsChecked, EShowMapsChoices::ShowAllMaps)
-							.OnCheckStateChanged(this, &SProjectLauncherCookByTheBookSettings::HandleShowCheckBoxCheckStateChanged, EShowMapsChoices::ShowAllMaps)
-							.Style(FEditorStyle::Get(), "RadioButton")
-							[
-								SNew(STextBlock)
-								.Text(LOCTEXT("AllMapsCheckBoxText", "Show all"))
-							]
+							SNew(SProjectLauncherFormLabel)
+								.LabelText(LOCTEXT("CookedMapsLabel", "Cooked Maps:"))
 						]
 
-						+ SHorizontalBox::Slot()
-							.FillWidth(1.0f)
-							.Padding(8.0f, 0.0f, 0.0f, 0.0f)
-							[
-								// cooked maps radio button
-								SNew(SCheckBox)
-								.IsChecked(this, &SProjectLauncherCookByTheBookSettings::HandleShowCheckBoxIsChecked, EShowMapsChoices::ShowCookedMaps)
-								.OnCheckStateChanged(this, &SProjectLauncherCookByTheBookSettings::HandleShowCheckBoxCheckStateChanged, EShowMapsChoices::ShowCookedMaps)
-								.Style(FEditorStyle::Get(), "RadioButton")
+					+ SVerticalBox::Slot()
+						.AutoHeight()
+						.Padding(0.0f, 4.0f)
+						.VAlign(VAlign_Center)
+						[
+							SNew(SHorizontalBox)
+
+							+ SHorizontalBox::Slot()
+								.AutoWidth()
 								[
-									SNew(STextBlock)
-									.Text(LOCTEXT("CookedMapsCheckBoxText", "Show cooked"))
+									// all maps radio button
+									SNew(SCheckBox)
+										.IsChecked(this, &SProjectLauncherCookByTheBookSettings::HandleShowCheckBoxIsChecked, EShowMapsChoices::ShowAllMaps)
+										.OnCheckStateChanged(this, &SProjectLauncherCookByTheBookSettings::HandleShowCheckBoxCheckStateChanged, EShowMapsChoices::ShowAllMaps)
+										.Style(FEditorStyle::Get(), "RadioButton")
+										[
+											SNew(STextBlock)
+												.Text(LOCTEXT("AllMapsCheckBoxText", "Show all"))
+										]
 								]
-							]
-					]
-
-				+ SVerticalBox::Slot()
-					.FillHeight(1.0)
-					.Padding(0.0f, 4.0f, 0.0f, 0.0f)
-					[
-						// map list
-						SAssignNew(MapListView, SListView<TSharedPtr<FString> >)
-						.HeaderRow
-						(
-						SNew(SHeaderRow)
-						.Visibility(EVisibility::Collapsed)
-
-						+ SHeaderRow::Column("MapName")
-						.DefaultLabel(LOCTEXT("MapListMapNameColumnHeader", "Map"))
-						.FillWidth(1.0f)
-						)
-						.ItemHeight(16.0f)
-						.ListItemsSource(&MapList)
-						.OnGenerateRow(this, &SProjectLauncherCookByTheBookSettings::HandleMapListViewGenerateRow)
-						.SelectionMode(ESelectionMode::None)
-					]
 
-				+ SVerticalBox::Slot()
-					.AutoHeight()
-					[
-						SNew(SHorizontalBox)
-						.Visibility(this, &SProjectLauncherCookByTheBookSettings::HandleNoMapSelectedBoxVisibility)
+							+ SHorizontalBox::Slot()
+								.FillWidth(1.0f)
+								.Padding(8.0f, 0.0f, 0.0f, 0.0f)
+								[
+									// cooked maps radio button
+									SNew(SCheckBox)
+										.IsChecked(this, &SProjectLauncherCookByTheBookSettings::HandleShowCheckBoxIsChecked, EShowMapsChoices::ShowCookedMaps)
+										.OnCheckStateChanged(this, &SProjectLauncherCookByTheBookSettings::HandleShowCheckBoxCheckStateChanged, EShowMapsChoices::ShowCookedMaps)
+										.Style(FEditorStyle::Get(), "RadioButton")
+										[
+											SNew(STextBlock)
+												.Text(LOCTEXT("CookedMapsCheckBoxText", "Show cooked"))
+										]
+								]
+						]
 
-						+ SHorizontalBox::Slot()
-						.AutoWidth()
+					+ SVerticalBox::Slot()
+						.FillHeight(1.0)
+						.Padding(0.0f, 4.0f, 0.0f, 0.0f)
 						[
-							SNew(SImage)
-							.Image(FEditorStyle::GetBrush(TEXT("Icons.Warning")))
+							// map list
+							SAssignNew(MapListView, SListView<TSharedPtr<FString> >)
+								.HeaderRow
+								(
+									SNew(SHeaderRow)
+										.Visibility(EVisibility::Collapsed)
+
+									+ SHeaderRow::Column("MapName")
+										.DefaultLabel(LOCTEXT("MapListMapNameColumnHeader", "Map"))
+										.FillWidth(1.0f)
+								)
+								.ItemHeight(16.0f)
+								.ListItemsSource(&MapList)
+								.OnGenerateRow(this, &SProjectLauncherCookByTheBookSettings::HandleMapListViewGenerateRow)
+								.SelectionMode(ESelectionMode::None)
 						]
 
-						+ SHorizontalBox::Slot()
+					+ SVerticalBox::Slot()
+						.AutoHeight()
+						[
+							SNew(SHorizontalBox)
+								.Visibility(this, &SProjectLauncherCookByTheBookSettings::HandleNoMapSelectedBoxVisibility)
+
+							+ SHorizontalBox::Slot()
 							.AutoWidth()
-							.Padding(4.0f, 0.0f)
-							.VAlign(VAlign_Center)
 							[
-								SNew(STextBlock)
-								.Text(this, &SProjectLauncherCookByTheBookSettings::HandleNoMapsTextBlockText)
+								SNew(SImage)
+									.Image(FEditorStyle::GetBrush(TEXT("Icons.Warning")))
 							]
-					]
-
-				+ SVerticalBox::Slot()
-					.AutoHeight()
-					.Padding(0.0f, 6.0f, 0.0f, 4.0f)
-					[
-						SNew(SSeparator)
-						.Orientation(Orient_Horizontal)
-					]
 
-				+ SVerticalBox::Slot()
-					.AutoHeight()
-					.VAlign(VAlign_Center)
-					[
-						SNew(SHorizontalBox)
+							+ SHorizontalBox::Slot()
+								.AutoWidth()
+								.Padding(4.0f, 0.0f)
+								.VAlign(VAlign_Center)
+								[
+									SNew(STextBlock)
+										.Text(this, &SProjectLauncherCookByTheBookSettings::HandleNoMapsTextBlockText)
+								]
+						]
 
-						+ SHorizontalBox::Slot()
-						.FillWidth(1.0f)
-						.HAlign(HAlign_Right)
+					+ SVerticalBox::Slot()
+						.AutoHeight()
+						.Padding(0.0f, 6.0f, 0.0f, 4.0f)
 						[
-							SNew(STextBlock)
-							.Text(LOCTEXT("SelectLabel", "Select:"))
-							.Visibility(this, &SProjectLauncherCookByTheBookSettings::HandleMapSelectionHyperlinkVisibility)
+							SNew(SSeparator)
+								.Orientation(Orient_Horizontal)
 						]
 
-						+ SHorizontalBox::Slot()
-							.AutoWidth()
-							.Padding(8.0f, 0.0f)
-							[
-								// all maps hyper link
-								SNew(SHyperlink)
-								.OnNavigate(this, &SProjectLauncherCookByTheBookSettings::HandleAllMapsHyperlinkNavigate, true)
-								.Text(LOCTEXT("AllMapsHyperlinkLabel", "All"))
-								.ToolTipText(LOCTEXT("AllMapsHyperlinkTooltip", "Select all available maps."))
-								.Visibility(this, &SProjectLauncherCookByTheBookSettings::HandleMapSelectionHyperlinkVisibility)									
-							]
+					+ SVerticalBox::Slot()
+						.AutoHeight()
+						.VAlign(VAlign_Center)
+						[
+							SNew(SHorizontalBox)
 
-						+ SHorizontalBox::Slot()
-							.AutoWidth()
-							[
-								// no maps hyper link
-								SNew(SHyperlink)
-								.OnNavigate(this, &SProjectLauncherCookByTheBookSettings::HandleAllMapsHyperlinkNavigate, false)
-								.Text(LOCTEXT("NoMapsHyperlinkLabel", "None"))
-								.ToolTipText(LOCTEXT("NoMapsHyperlinkTooltip", "Deselect all maps."))
-								.Visibility(this, &SProjectLauncherCookByTheBookSettings::HandleMapSelectionHyperlinkVisibility)
-							]
-					]
-			]
+							+ SHorizontalBox::Slot()
+								.FillWidth(1.0f)
+								.HAlign(HAlign_Right)
+								[
+									SNew(STextBlock)
+										.Text(LOCTEXT("SelectLabel", "Select:"))
+										.Visibility(this, &SProjectLauncherCookByTheBookSettings::HandleMapSelectionHyperlinkVisibility)
+								]
+
+							+ SHorizontalBox::Slot()
+								.AutoWidth()
+								.Padding(8.0f, 0.0f)
+								[
+									// all maps hyper link
+									SNew(SHyperlink)
+										.OnNavigate(this, &SProjectLauncherCookByTheBookSettings::HandleAllMapsHyperlinkNavigate, true)
+										.Text(LOCTEXT("AllMapsHyperlinkLabel", "All"))
+										.ToolTipText(LOCTEXT("AllMapsHyperlinkTooltip", "Select all available maps."))
+										.Visibility(this, &SProjectLauncherCookByTheBookSettings::HandleMapSelectionHyperlinkVisibility)
+								]
+
+							+ SHorizontalBox::Slot()
+								.AutoWidth()
+								[
+									// no maps hyper link
+									SNew(SHyperlink)
+										.OnNavigate(this, &SProjectLauncherCookByTheBookSettings::HandleAllMapsHyperlinkNavigate, false)
+										.Text(LOCTEXT("NoMapsHyperlinkLabel", "None"))
+										.ToolTipText(LOCTEXT("NoMapsHyperlinkTooltip", "Deselect all maps."))
+										.Visibility(this, &SProjectLauncherCookByTheBookSettings::HandleMapSelectionHyperlinkVisibility)
+								]
+						]
+				]
 		]
 
 	+ SVerticalBox::Slot()
@@ -1046,124 +1057,125 @@ TSharedRef<SWidget> SProjectLauncherCookByTheBookSettings::MakeSimpleWidget()
 		.Padding(0.0f, 8.0f, 0.0f, 0.0f)
 		[
 			SNew(SExpandableArea)
-			.AreaTitle(LOCTEXT("AdvancedAreaTitle", "Advanced Settings"))
-			.InitiallyCollapsed(true)
-			.Padding(8.0f)
-			.BodyContent()
-			[
-				SNew(SVerticalBox)
-
-				+ SVerticalBox::Slot()
-				.AutoHeight()
+				.AreaTitle(LOCTEXT("AdvancedAreaTitle", "Advanced Settings"))
+				.InitiallyCollapsed(true)
+				.Padding(8.0f)
+				.BodyContent()
 				[
-					// incremental cook check box
-					SNew(SCheckBox)
-					.IsChecked(this, &SProjectLauncherCookByTheBookSettings::HandleIncrementalCheckBoxIsChecked)
-					.OnCheckStateChanged(this, &SProjectLauncherCookByTheBookSettings::HandleIncrementalCheckBoxCheckStateChanged)
-					.Padding(FMargin(4.0f, 0.0f))
-					.ToolTipText(LOCTEXT("IncrementalCheckBoxTooltip", "If checked, only modified content will be cooked, resulting in much faster cooking times. It is recommended to enable this option whenever possible."))
-					.Content()
-					[
-						SNew(STextBlock)
-						.Text(LOCTEXT("IncrementalCheckBoxText", "Iterative cooking: Only cook content modified from previous cook"))
-					]
-				]
+					SNew(SVerticalBox)
 
-				// disabled for now until this system is live
-				/*+ SVerticalBox::Slot()
-					.AutoHeight()
-					[
-						// incremental cook check box
-						SNew(SCheckBox)
-						.IsChecked(this, &SProjectLauncherCookByTheBookSettings::HandleSharedCookedBuildCheckBoxIsChecked)
-						.OnCheckStateChanged(this, &SProjectLauncherCookByTheBookSettings::HandleSharedCookedBuildCheckBoxCheckStateChanged)
-						.Padding(FMargin(4.0f, 0.0f))
-						.ToolTipText(LOCTEXT("SharedCookedBuildCheckBoxToolTip", "Experimental: Use a build from the network to cook from."))
-						.Content()
+					+ SVerticalBox::Slot()
+						.AutoHeight()
 						[
-							SNew(STextBlock)
-							.Text(LOCTEXT("SharedCookedBuildCheckBoxText", "Iteratively cook from a pre packaged build located on the network"))
+							// incremental cook check box
+							SNew(SCheckBox)
+								.IsChecked(this, &SProjectLauncherCookByTheBookSettings::HandleIncrementalCheckBoxIsChecked)
+								.OnCheckStateChanged(this, &SProjectLauncherCookByTheBookSettings::HandleIncrementalCheckBoxCheckStateChanged)
+								.Padding(FMargin(4.0f, 0.0f))
+								.ToolTipText(LOCTEXT("IncrementalCheckBoxTooltip", "If checked, only modified content will be cooked, resulting in much faster cooking times. It is recommended to enable this option whenever possible."))
+								.Content()
+								[
+									SNew(STextBlock)
+										.Text(LOCTEXT("IncrementalCheckBoxText", "Iterative cooking: Only cook content modified from previous cook"))
+								]
 						]
-					]*/
 
-				+ SVerticalBox::Slot()
-					.AutoHeight()
-					.Padding(0.0f, 4.0f, 0.0f, 0.0f)
-					[
-						// incremental cook check box
-						SNew(SCheckBox)
-						.IsChecked(this, &SProjectLauncherCookByTheBookSettings::HandleUnversionedCheckBoxIsChecked)
-						.OnCheckStateChanged(this, &SProjectLauncherCookByTheBookSettings::HandleUnversionedCheckBoxCheckStateChanged)
-						.Padding(FMargin(4.0f, 0.0f))
-						.ToolTipText(LOCTEXT("UnversionedCheckBoxTooltip", "If checked, the version is assumed to be current at load. This is potentially dangerous, but results in smaller patch sizes."))
-						.Content()
+					// disabled for now until this system is live
+					/*+ SVerticalBox::Slot()
+						.AutoHeight()
+						[
+							// incremental cook check box
+							SNew(SCheckBox)
+							.IsChecked(this, &SProjectLauncherCookByTheBookSettings::HandleSharedCookedBuildCheckBoxIsChecked)
+							.OnCheckStateChanged(this, &SProjectLauncherCookByTheBookSettings::HandleSharedCookedBuildCheckBoxCheckStateChanged)
+							.Padding(FMargin(4.0f, 0.0f))
+							.ToolTipText(LOCTEXT("SharedCookedBuildCheckBoxToolTip", "Experimental: Use a build from the network to cook from."))
+							.Content()
+							[
+								SNew(STextBlock)
+								.Text(LOCTEXT("SharedCookedBuildCheckBoxText", "Iteratively cook from a pre packaged build located on the network"))
+							]
+						]*/
+
+					+ SVerticalBox::Slot()
+						.AutoHeight()
+						.Padding(0.0f, 4.0f, 0.0f, 0.0f)
 						[
-							SNew(STextBlock)
-							.Text(LOCTEXT("UnversionedCheckBoxText", "Save packages without versions"))
+							// incremental cook check box
+							SNew(SCheckBox)
+								.IsChecked(this, &SProjectLauncherCookByTheBookSettings::HandleUnversionedCheckBoxIsChecked)
+								.OnCheckStateChanged(this, &SProjectLauncherCookByTheBookSettings::HandleUnversionedCheckBoxCheckStateChanged)
+								.Padding(FMargin(4.0f, 0.0f))
+								.ToolTipText(LOCTEXT("UnversionedCheckBoxTooltip", "If checked, the version is assumed to be current at load. This is potentially dangerous, but results in smaller patch sizes."))
+								.Content()
+								[
+									SNew(STextBlock)
+										.Text(LOCTEXT("UnversionedCheckBoxText", "Save packages without versions"))
+								]
 						]
-					]
 
-				+ SVerticalBox::Slot()
-					.AutoHeight()
-					.Padding(0.0f, 4.0f, 0.0f, 0.0f)
-					[
-						// unreal pak check box
-						SNew(SCheckBox)
-						.IsChecked(this, &SProjectLauncherCookByTheBookSettings::HandleUnrealPakCheckBoxIsChecked)
-						.OnCheckStateChanged(this, &SProjectLauncherCookByTheBookSettings::HandleUnrealPakCheckBoxCheckStateChanged)
-						.Padding(FMargin(4.0f, 0.0f))
-						.ToolTipText(LOCTEXT("UnrealPakCheckBoxTooltip", "If checked, the content will be deployed as a single UnrealPak file instead of many separate files."))
-						.Content()
+					+ SVerticalBox::Slot()
+						.AutoHeight()
+						.Padding(0.0f, 4.0f, 0.0f, 0.0f)
 						[
-							SNew(STextBlock)
-							.Text(LOCTEXT("UnrealPakCheckBoxText", "Store all content in a single file (UnrealPak)"))
+							// unreal pak check box
+							SNew(SCheckBox)
+								.IsChecked(this, &SProjectLauncherCookByTheBookSettings::HandleUnrealPakCheckBoxIsChecked)
+								.OnCheckStateChanged(this, &SProjectLauncherCookByTheBookSettings::HandleUnrealPakCheckBoxCheckStateChanged)
+								.Padding(FMargin(4.0f, 0.0f))
+								.ToolTipText(LOCTEXT("UnrealPakCheckBoxTooltip", "If checked, the content will be deployed as a single UnrealPak file instead of many separate files."))
+								.Content()
+								[
+									SNew(STextBlock)
+										.Text(LOCTEXT("UnrealPakCheckBoxText", "Store all content in a single file (UnrealPak)"))
+								]
 						]
-					]
 
-				+ SVerticalBox::Slot()
-					.AutoHeight()
-					.Padding(0.0f, 12.0f, 0.0f, 0.0f)
-					[
-						SNew(SProjectLauncherFormLabel)
-						.LabelText(LOCTEXT("CookConfigurationSelectorLabel", "Cooker build configuration:"))							
-					]
+					+ SVerticalBox::Slot()
+						.AutoHeight()
+						.Padding(0.0f, 12.0f, 0.0f, 0.0f)
+						[
+							SNew(SProjectLauncherFormLabel)
+								.LabelText(LOCTEXT("CookConfigurationSelectorLabel", "Cooker build configuration:"))							
+						]
 
-				+ SVerticalBox::Slot()
-					.AutoHeight()
-					.Padding(0.0f, 4.0f, 0.0f, 0.0f)
-					[
-						// cooker build configuration selector
-						SNew(SProjectLauncherBuildConfigurationSelector)
-						.OnConfigurationSelected(this, &SProjectLauncherCookByTheBookSettings::HandleCookConfigurationSelectorConfigurationSelected)
-						.Text(this, &SProjectLauncherCookByTheBookSettings::HandleCookConfigurationSelectorText)
-						.ToolTipText(LOCTEXT("CookConfigurationToolTipText", "Sets the build configuration to use for the cooker commandlet."))
-					]
+					+ SVerticalBox::Slot()
+						.AutoHeight()
+						.Padding(0.0f, 4.0f, 0.0f, 0.0f)
+						[
+							// cooker build configuration selector
+							SNew(SProjectLauncherBuildConfigurationSelector)
+								.OnConfigurationSelected(this, &SProjectLauncherCookByTheBookSettings::HandleCookConfigurationSelectorConfigurationSelected)
+								.Text(this, &SProjectLauncherCookByTheBookSettings::HandleCookConfigurationSelectorText)
+								.ToolTipText(LOCTEXT("CookConfigurationToolTipText", "Sets the build configuration to use for the cooker commandlet."))
+						]
 
-				+ SVerticalBox::Slot()
-					.AutoHeight()
-					.Padding(0.0f, 8.0f, 0.0f, 0.0f)
-					[
-						SNew(SProjectLauncherFormLabel)
-						.LabelText(LOCTEXT("CookerOptionsTextBoxLabel", "Additional Cooker Options:"))
-					]
+					+ SVerticalBox::Slot()
+						.AutoHeight()
+						.Padding(0.0f, 8.0f, 0.0f, 0.0f)
+						[
+							SNew(SProjectLauncherFormLabel)
+								.LabelText(LOCTEXT("CookerOptionsTextBoxLabel", "Additional Cooker Options:"))
+						]
 
-				+ SVerticalBox::Slot()
-					.AutoHeight()
-					.Padding(0.0f, 4.0f, 0.0f, 0.0f)
-					[
-						// cooker command line options
-						SNew(SEditableTextBox)
-						.ToolTipText(LOCTEXT("CookerOptionsTextBoxTooltip", "Additional cooker command line parameters can be specified here."))
-						.Text(this, &SProjectLauncherCookByTheBookSettings::HandleCookOptionsTextBlockText)
-						.OnTextCommitted(this, &SProjectLauncherCookByTheBookSettings::HandleCookerOptionsCommitted)
-					]
-			]
+					+ SVerticalBox::Slot()
+						.AutoHeight()
+						.Padding(0.0f, 4.0f, 0.0f, 0.0f)
+						[
+							// cooker command line options
+							SNew(SEditableTextBox)
+								.ToolTipText(LOCTEXT("CookerOptionsTextBoxTooltip", "Additional cooker command line parameters can be specified here."))
+								.Text(this, &SProjectLauncherCookByTheBookSettings::HandleCookOptionsTextBlockText)
+								.OnTextCommitted(this, &SProjectLauncherCookByTheBookSettings::HandleCookerOptionsCommitted)
+						]
+				]
 		];
 
 	return Widget;
 }
 END_SLATE_FUNCTION_BUILD_OPTIMIZATION
 
+
 void SProjectLauncherCookByTheBookSettings::RefreshCultureList()
 {
 	CultureList.Reset();
@@ -1186,7 +1198,8 @@ void SProjectLauncherCookByTheBookSettings::RefreshCultureList()
 	}
 }
 
-void SProjectLauncherCookByTheBookSettings::RefreshMapList( )
+
+void SProjectLauncherCookByTheBookSettings::RefreshMapList()
 {
 	MapList.Reset();
 
@@ -1214,7 +1227,7 @@ void SProjectLauncherCookByTheBookSettings::RefreshMapList( )
 /* SProjectLauncherCookByTheBookSettings callbacks
  *****************************************************************************/
 
-void SProjectLauncherCookByTheBookSettings::HandleAllCulturesHyperlinkNavigate( bool AllPlatforms )
+void SProjectLauncherCookByTheBookSettings::HandleAllCulturesHyperlinkNavigate(bool AllPlatforms)
 {
 	ILauncherProfilePtr SelectedProfile = Model->GetSelectedProfile();
 
@@ -1238,7 +1251,7 @@ void SProjectLauncherCookByTheBookSettings::HandleAllCulturesHyperlinkNavigate(
 }
 
 
-EVisibility SProjectLauncherCookByTheBookSettings::HandleAllCulturesHyperlinkVisibility( ) const
+EVisibility SProjectLauncherCookByTheBookSettings::HandleAllCulturesHyperlinkVisibility() const
 {
 	TArray<FString> CultureNames;
 	FInternationalization::Get().GetCultureNames(CultureNames);
@@ -1252,7 +1265,7 @@ EVisibility SProjectLauncherCookByTheBookSettings::HandleAllCulturesHyperlinkVis
 }
 
 
-void SProjectLauncherCookByTheBookSettings::HandleAllMapsHyperlinkNavigate( bool AllPlatforms )
+void SProjectLauncherCookByTheBookSettings::HandleAllMapsHyperlinkNavigate(bool AllPlatforms)
 {
 	ILauncherProfilePtr SelectedProfile = Model->GetSelectedProfile();
 
@@ -1275,7 +1288,7 @@ void SProjectLauncherCookByTheBookSettings::HandleAllMapsHyperlinkNavigate( bool
 }
 
 
-EVisibility SProjectLauncherCookByTheBookSettings::HandleMapSelectionHyperlinkVisibility( ) const
+EVisibility SProjectLauncherCookByTheBookSettings::HandleMapSelectionHyperlinkVisibility() const
 {
 	if (MapList.Num() > 1)
 	{
@@ -1286,7 +1299,7 @@ EVisibility SProjectLauncherCookByTheBookSettings::HandleMapSelectionHyperlinkVi
 }
 
 
-void SProjectLauncherCookByTheBookSettings::HandleCookConfigurationSelectorConfigurationSelected( EBuildConfigurations::Type Configuration )
+void SProjectLauncherCookByTheBookSettings::HandleCookConfigurationSelectorConfigurationSelected(EBuildConfigurations::Type Configuration)
 {
 	ILauncherProfilePtr SelectedProfile = Model->GetSelectedProfile();
 
@@ -1297,7 +1310,7 @@ void SProjectLauncherCookByTheBookSettings::HandleCookConfigurationSelectorConfi
 }
 
 
-FText SProjectLauncherCookByTheBookSettings::HandleCookConfigurationSelectorText( ) const
+FText SProjectLauncherCookByTheBookSettings::HandleCookConfigurationSelectorText() const
 {
 	ILauncherProfilePtr SelectedProfile = Model->GetSelectedProfile();
 
@@ -1310,7 +1323,7 @@ FText SProjectLauncherCookByTheBookSettings::HandleCookConfigurationSelectorText
 }
 
 
-void SProjectLauncherCookByTheBookSettings::HandleIncrementalCheckBoxCheckStateChanged( ECheckBoxState NewState )
+void SProjectLauncherCookByTheBookSettings::HandleIncrementalCheckBoxCheckStateChanged(ECheckBoxState NewState)
 {
 	ILauncherProfilePtr SelectedProfile = Model->GetSelectedProfile();
 
@@ -1321,7 +1334,7 @@ void SProjectLauncherCookByTheBookSettings::HandleIncrementalCheckBoxCheckStateC
 }
 
 
-ECheckBoxState SProjectLauncherCookByTheBookSettings::HandleIncrementalCheckBoxIsChecked( ) const
+ECheckBoxState SProjectLauncherCookByTheBookSettings::HandleIncrementalCheckBoxIsChecked() const
 {
 	ILauncherProfilePtr SelectedProfile = Model->GetSelectedProfile();
 
@@ -1364,9 +1377,7 @@ ECheckBoxState SProjectLauncherCookByTheBookSettings::HandleSharedCookedBuildChe
 }
 
 
-
-
-void SProjectLauncherCookByTheBookSettings::HandleCompressedCheckBoxCheckStateChanged( ECheckBoxState NewState )
+void SProjectLauncherCookByTheBookSettings::HandleCompressedCheckBoxCheckStateChanged(ECheckBoxState NewState)
 {
 	ILauncherProfilePtr SelectedProfile = Model->GetSelectedProfile();
 
@@ -1377,7 +1388,7 @@ void SProjectLauncherCookByTheBookSettings::HandleCompressedCheckBoxCheckStateCh
 }
 
 
-ECheckBoxState SProjectLauncherCookByTheBookSettings::HandleCompressedCheckBoxIsChecked( ) const
+ECheckBoxState SProjectLauncherCookByTheBookSettings::HandleCompressedCheckBoxIsChecked() const
 {
 	ILauncherProfilePtr SelectedProfile = Model->GetSelectedProfile();
 
@@ -1393,7 +1404,6 @@ ECheckBoxState SProjectLauncherCookByTheBookSettings::HandleCompressedCheckBoxIs
 }
 
 
-
 void SProjectLauncherCookByTheBookSettings::HandleEncryptIniFilesCheckBoxCheckStateChanged(ECheckBoxState NewState)
 {
 	ILauncherProfilePtr SelectedProfile = Model->GetSelectedProfile();
@@ -1420,7 +1430,8 @@ ECheckBoxState SProjectLauncherCookByTheBookSettings::HandleEncryptIniFilesCheck
 	return ECheckBoxState::Unchecked;
 }
 
-TSharedRef<ITableRow> SProjectLauncherCookByTheBookSettings::HandleMapListViewGenerateRow( TSharedPtr<FString> InItem, const TSharedRef<STableViewBase>& OwnerTable )
+
+TSharedRef<ITableRow> SProjectLauncherCookByTheBookSettings::HandleMapListViewGenerateRow(TSharedPtr<FString> InItem, const TSharedRef<STableViewBase>& OwnerTable)
 {
 	return SNew(SProjectLauncherMapListRow, Model.ToSharedRef())
 		.MapName(InItem)
@@ -1428,7 +1439,7 @@ TSharedRef<ITableRow> SProjectLauncherCookByTheBookSettings::HandleMapListViewGe
 }
 
 
-TSharedRef<ITableRow> SProjectLauncherCookByTheBookSettings::HandleCultureListViewGenerateRow( TSharedPtr<FString> InItem, const TSharedRef<STableViewBase>& OwnerTable )
+TSharedRef<ITableRow> SProjectLauncherCookByTheBookSettings::HandleCultureListViewGenerateRow(TSharedPtr<FString> InItem, const TSharedRef<STableViewBase>& OwnerTable)
 {
 	return SNew(SProjectLauncherCultureListRow, Model.ToSharedRef())
 		.CultureName(InItem)
@@ -1436,7 +1447,7 @@ TSharedRef<ITableRow> SProjectLauncherCookByTheBookSettings::HandleCultureListVi
 }
 
 
-EVisibility SProjectLauncherCookByTheBookSettings::HandleNoMapSelectedBoxVisibility( ) const
+EVisibility SProjectLauncherCookByTheBookSettings::HandleNoMapSelectedBoxVisibility() const
 {
 	if (MapList.Num() == 0)
 	{
@@ -1447,7 +1458,7 @@ EVisibility SProjectLauncherCookByTheBookSettings::HandleNoMapSelectedBoxVisibil
 }
 
 
-FText SProjectLauncherCookByTheBookSettings::HandleNoMapsTextBlockText( ) const
+FText SProjectLauncherCookByTheBookSettings::HandleNoMapsTextBlockText() const
 {
 	if (MapList.Num() == 0)
 	{
@@ -1465,7 +1476,7 @@ FText SProjectLauncherCookByTheBookSettings::HandleNoMapsTextBlockText( ) const
 }
 
 
-void SProjectLauncherCookByTheBookSettings::HandleProfileManagerProfileSelected( const ILauncherProfilePtr& SelectedProfile, const ILauncherProfilePtr& PreviousProfile )
+void SProjectLauncherCookByTheBookSettings::HandleProfileManagerProfileSelected(const ILauncherProfilePtr& SelectedProfile, const ILauncherProfilePtr& PreviousProfile)
 {
 	if (PreviousProfile.IsValid())
 	{
@@ -1479,13 +1490,15 @@ void SProjectLauncherCookByTheBookSettings::HandleProfileManagerProfileSelected(
 	RefreshCultureList();
 }
 
+
 void SProjectLauncherCookByTheBookSettings::HandleProfileProjectChanged()
 {
 	RefreshMapList();
 	RefreshCultureList();
 }
 
-ECheckBoxState SProjectLauncherCookByTheBookSettings::HandleShowCheckBoxIsChecked( EShowMapsChoices::Type Choice ) const
+
+ECheckBoxState SProjectLauncherCookByTheBookSettings::HandleShowCheckBoxIsChecked(EShowMapsChoices::Type Choice) const
 {
 	if (ShowMapsChoice == Choice)
 	{
@@ -1496,7 +1509,7 @@ ECheckBoxState SProjectLauncherCookByTheBookSettings::HandleShowCheckBoxIsChecke
 }
 
 
-void SProjectLauncherCookByTheBookSettings::HandleShowCheckBoxCheckStateChanged( ECheckBoxState NewState, EShowMapsChoices::Type Choice )
+void SProjectLauncherCookByTheBookSettings::HandleShowCheckBoxCheckStateChanged(ECheckBoxState NewState, EShowMapsChoices::Type Choice)
 {
 	if (NewState == ECheckBoxState::Checked)
 	{
@@ -1506,7 +1519,7 @@ void SProjectLauncherCookByTheBookSettings::HandleShowCheckBoxCheckStateChanged(
 }
 
 
-void SProjectLauncherCookByTheBookSettings::HandleUnversionedCheckBoxCheckStateChanged( ECheckBoxState NewState )
+void SProjectLauncherCookByTheBookSettings::HandleUnversionedCheckBoxCheckStateChanged(ECheckBoxState NewState)
 {
 	ILauncherProfilePtr SelectedProfile = Model->GetSelectedProfile();
 
@@ -1517,7 +1530,7 @@ void SProjectLauncherCookByTheBookSettings::HandleUnversionedCheckBoxCheckStateC
 }
 
 
-ECheckBoxState SProjectLauncherCookByTheBookSettings::HandleUnversionedCheckBoxIsChecked( ) const
+ECheckBoxState SProjectLauncherCookByTheBookSettings::HandleUnversionedCheckBoxIsChecked() const
 {
 	ILauncherProfilePtr SelectedProfile = Model->GetSelectedProfile();
 
@@ -1533,7 +1546,7 @@ ECheckBoxState SProjectLauncherCookByTheBookSettings::HandleUnversionedCheckBoxI
 }
 
 
-EVisibility SProjectLauncherCookByTheBookSettings::HandleValidationErrorIconVisibility( ELauncherProfileValidationErrors::Type Error ) const
+EVisibility SProjectLauncherCookByTheBookSettings::HandleValidationErrorIconVisibility(ELauncherProfileValidationErrors::Type Error) const
 {
 	ILauncherProfilePtr SelectedProfile = Model->GetSelectedProfile();
 
@@ -1548,6 +1561,7 @@ EVisibility SProjectLauncherCookByTheBookSettings::HandleValidationErrorIconVisi
 	return EVisibility::Hidden;
 }
 
+
 FText SProjectLauncherCookByTheBookSettings::HandleCookOptionsTextBlockText() const
 {
 	ILauncherProfilePtr SelectedProfile = Model->GetSelectedProfile();
@@ -1562,6 +1576,7 @@ FText SProjectLauncherCookByTheBookSettings::HandleCookOptionsTextBlockText() co
 	return result;
 }
 
+
 void SProjectLauncherCookByTheBookSettings::HandleCookerOptionsCommitted(const FText& NewText, ETextCommit::Type CommitType)
 {
 	ILauncherProfilePtr SelectedProfile = Model->GetSelectedProfile();
@@ -1583,7 +1598,6 @@ void SProjectLauncherCookByTheBookSettings::HandleCookerOptionsCommitted(const F
 }
 
 
-
 FText SProjectLauncherCookByTheBookSettings::HandleMultiProcessCookerTextBlockText() const
 {
 	ILauncherProfilePtr SelectedProfile = Model->GetSelectedProfile();
@@ -1598,6 +1612,7 @@ FText SProjectLauncherCookByTheBookSettings::HandleMultiProcessCookerTextBlockTe
 	return result;
 }
 
+
 void SProjectLauncherCookByTheBookSettings::HandleMultiProcessCookerCommitted(const FText& NewText, ETextCommit::Type CommitType)
 {
 	ILauncherProfilePtr SelectedProfile = Model->GetSelectedProfile();
@@ -1619,9 +1634,7 @@ void SProjectLauncherCookByTheBookSettings::HandleMultiProcessCookerCommitted(co
 }
 
 
-
-
-void SProjectLauncherCookByTheBookSettings::HandleUnrealPakCheckBoxCheckStateChanged( ECheckBoxState NewState )
+void SProjectLauncherCookByTheBookSettings::HandleUnrealPakCheckBoxCheckStateChanged(ECheckBoxState NewState)
 {
 	ILauncherProfilePtr SelectedProfile = Model->GetSelectedProfile();
 
@@ -1632,7 +1645,7 @@ void SProjectLauncherCookByTheBookSettings::HandleUnrealPakCheckBoxCheckStateCha
 }
 
 
-ECheckBoxState SProjectLauncherCookByTheBookSettings::HandleUnrealPakCheckBoxIsChecked( ) const
+ECheckBoxState SProjectLauncherCookByTheBookSettings::HandleUnrealPakCheckBoxIsChecked() const
 {
 	ILauncherProfilePtr SelectedProfile = Model->GetSelectedProfile();
 
@@ -1647,8 +1660,8 @@ ECheckBoxState SProjectLauncherCookByTheBookSettings::HandleUnrealPakCheckBoxIsC
 	return ECheckBoxState::Unchecked;
 }
 
-// Callback for check state changes of the 'UnrealPak' check box.
-void SProjectLauncherCookByTheBookSettings::HandleGeneratePatchCheckBoxCheckStateChanged( ECheckBoxState NewState )
+
+void SProjectLauncherCookByTheBookSettings::HandleGeneratePatchCheckBoxCheckStateChanged(ECheckBoxState NewState)
 {
 	ILauncherProfilePtr SelectedProfile = Model->GetSelectedProfile();
 
@@ -1658,7 +1671,7 @@ void SProjectLauncherCookByTheBookSettings::HandleGeneratePatchCheckBoxCheckStat
 	}
 }
 
-// Callback for check state changes of the 'AddPatchLevel' check box.
+
 void SProjectLauncherCookByTheBookSettings::HandleAddPatchLevelCheckBoxCheckStateChanged( ECheckBoxState NewState )
 {
 	ILauncherProfilePtr SelectedProfile = Model->GetSelectedProfile();
@@ -1669,7 +1682,7 @@ void SProjectLauncherCookByTheBookSettings::HandleAddPatchLevelCheckBoxCheckStat
 	}
 }
 
-// Callback for check state changes of the 'StageBaseReleasePaks' check box.
+
 void SProjectLauncherCookByTheBookSettings::HandleStageBaseReleasePaksCheckBoxCheckStateChanged( ECheckBoxState NewState )
 {
 	ILauncherProfilePtr SelectedProfile = Model->GetSelectedProfile();
@@ -1680,8 +1693,8 @@ void SProjectLauncherCookByTheBookSettings::HandleStageBaseReleasePaksCheckBoxCh
 	}
 }
 
-// Callback for determining the checked state of the 'UnrealPak' check box.
-ECheckBoxState SProjectLauncherCookByTheBookSettings::HandleGeneratePatchCheckBoxIsChecked( ) const
+
+ECheckBoxState SProjectLauncherCookByTheBookSettings::HandleGeneratePatchCheckBoxIsChecked() const
 {
 	ILauncherProfilePtr SelectedProfile = Model->GetSelectedProfile();
 
@@ -1712,7 +1725,7 @@ ECheckBoxState SProjectLauncherCookByTheBookSettings::HandleAddPatchLevelCheckBo
 	return ECheckBoxState::Unchecked;
 }
 
-// Callback for determining the checked state of the 'StageBaseReleasePaks' check box.
+
 ECheckBoxState SProjectLauncherCookByTheBookSettings::HandleStageBaseReleasePaksCheckBoxIsChecked( ) const
 {
 	ILauncherProfilePtr SelectedProfile = Model->GetSelectedProfile();
@@ -1729,10 +1742,7 @@ ECheckBoxState SProjectLauncherCookByTheBookSettings::HandleStageBaseReleasePaks
 }
 
 
-
-
-// Callback for check state changes of the 'UnrealPak' check box.
-void SProjectLauncherCookByTheBookSettings::HandleCreateReleaseVersionCheckBoxCheckStateChanged( ECheckBoxState NewState )
+void SProjectLauncherCookByTheBookSettings::HandleCreateReleaseVersionCheckBoxCheckStateChanged(ECheckBoxState NewState)
 {
 	ILauncherProfilePtr SelectedProfile = Model->GetSelectedProfile();
 
@@ -1742,8 +1752,8 @@ void SProjectLauncherCookByTheBookSettings::HandleCreateReleaseVersionCheckBoxCh
 	}
 }
 
-// Callback for determining the checked state of the 'UnrealPak' check box.
-ECheckBoxState SProjectLauncherCookByTheBookSettings::HandleCreateReleaseVersionCheckBoxIsChecked( ) const
+
+ECheckBoxState SProjectLauncherCookByTheBookSettings::HandleCreateReleaseVersionCheckBoxIsChecked() const
 {
 	ILauncherProfilePtr SelectedProfile = Model->GetSelectedProfile();
 
@@ -1758,6 +1768,7 @@ ECheckBoxState SProjectLauncherCookByTheBookSettings::HandleCreateReleaseVersion
 	return ECheckBoxState::Unchecked;
 }
 
+
 FText SProjectLauncherCookByTheBookSettings::HandleCreateReleaseVersionNameTextBlockText() const
 {
 	ILauncherProfilePtr SelectedProfile = Model->GetSelectedProfile();
@@ -1772,6 +1783,7 @@ FText SProjectLauncherCookByTheBookSettings::HandleCreateReleaseVersionNameTextB
 	return result;
 }
 
+
 void SProjectLauncherCookByTheBookSettings::HandleCreateReleaseVersionNameCommitted(const FText& NewText, ETextCommit::Type CommitType)
 {
 	ILauncherProfilePtr SelectedProfile = Model->GetSelectedProfile();
@@ -1798,6 +1810,7 @@ FText SProjectLauncherCookByTheBookSettings::HandleBasedOnReleaseVersionNameText
 	return result;
 }
 
+
 void SProjectLauncherCookByTheBookSettings::HandleBasedOnReleaseVersionNameCommitted(const FText& NewText, ETextCommit::Type CommitType)
 {
 	ILauncherProfilePtr SelectedProfile = Model->GetSelectedProfile();
@@ -1808,6 +1821,8 @@ void SProjectLauncherCookByTheBookSettings::HandleBasedOnReleaseVersionNameCommi
 	}
 
 }
+
+
 FText SProjectLauncherCookByTheBookSettings::HandleDLCNameTextBlockText() const
 {
 	ILauncherProfilePtr SelectedProfile = Model->GetSelectedProfile();
@@ -1822,6 +1837,7 @@ FText SProjectLauncherCookByTheBookSettings::HandleDLCNameTextBlockText() const
 	return result;
 }
 
+
 void SProjectLauncherCookByTheBookSettings::HandleDLCNameCommitted(const FText& NewText, ETextCommit::Type CommitType)
 {
 	ILauncherProfilePtr SelectedProfile = Model->GetSelectedProfile();
@@ -1832,8 +1848,8 @@ void SProjectLauncherCookByTheBookSettings::HandleDLCNameCommitted(const FText&
 	}
 }
 
-// Callback for check state changes of the 'UnrealPak' check box.
-void SProjectLauncherCookByTheBookSettings::HandleBuildDLCCheckBoxCheckStateChanged( ECheckBoxState NewState ) 
+
+void SProjectLauncherCookByTheBookSettings::HandleBuildDLCCheckBoxCheckStateChanged(ECheckBoxState NewState) 
 {
 	ILauncherProfilePtr SelectedProfile = Model->GetSelectedProfile();
 
@@ -1843,7 +1859,7 @@ void SProjectLauncherCookByTheBookSettings::HandleBuildDLCCheckBoxCheckStateChan
 	}
 }
 
-// Callback for determining the checked state of the 'UnrealPak' check box.
+
 ECheckBoxState SProjectLauncherCookByTheBookSettings::HandleBuildDLCCheckBoxIsChecked() const
 {
 	ILauncherProfilePtr SelectedProfile = Model->GetSelectedProfile();
@@ -1859,6 +1875,7 @@ ECheckBoxState SProjectLauncherCookByTheBookSettings::HandleBuildDLCCheckBoxIsCh
 	return ECheckBoxState::Unchecked;
 }
 
+
 FReply SProjectLauncherCookByTheBookSettings::HandleHtppChunkInstallBrowseButtonClicked()
 {
 	IDesktopPlatform* DesktopPlatform = FDesktopPlatformModule::Get();
@@ -1895,6 +1912,7 @@ FReply SProjectLauncherCookByTheBookSettings::HandleHtppChunkInstallBrowseButton
 	return FReply::Handled();
 }
 
+
 FText SProjectLauncherCookByTheBookSettings::HandleHtppChunkInstallDirectoryText() const
 {
 	ILauncherProfilePtr SelectedProfile = Model->GetSelectedProfile();
@@ -1907,6 +1925,7 @@ FText SProjectLauncherCookByTheBookSettings::HandleHtppChunkInstallDirectoryText
 	return FText::GetEmpty();
 }
 
+
 void SProjectLauncherCookByTheBookSettings::HandleHtppChunkInstallDirectoryTextChanged(const FText& InText)
 {
 	ILauncherProfilePtr SelectedProfile = Model->GetSelectedProfile();
@@ -1931,6 +1950,7 @@ void SProjectLauncherCookByTheBookSettings::HandleHtppChunkInstallDirectoryTextC
 	}
 }
 
+
 void SProjectLauncherCookByTheBookSettings::HandleHttpChunkInstallCheckBoxCheckStateChanged(ECheckBoxState NewState)
 {
 	ILauncherProfilePtr SelectedProfile = Model->GetSelectedProfile();
@@ -1940,6 +1960,7 @@ void SProjectLauncherCookByTheBookSettings::HandleHttpChunkInstallCheckBoxCheckS
 	}
 }
 
+
 ECheckBoxState SProjectLauncherCookByTheBookSettings::HandleHttpChunkInstallCheckBoxIsChecked() const
 {
 	ILauncherProfilePtr SelectedProfile = Model->GetSelectedProfile();
@@ -1947,9 +1968,11 @@ ECheckBoxState SProjectLauncherCookByTheBookSettings::HandleHttpChunkInstallChec
 	{
 		return SelectedProfile->IsGenerateHttpChunkData() ? ECheckBoxState::Checked : ECheckBoxState::Unchecked;
 	}
+
 	return ECheckBoxState::Unchecked;
 }
 
+
 FText SProjectLauncherCookByTheBookSettings::HandleHttpChunkInstallNameTextBlockText() const
 {
 	ILauncherProfilePtr SelectedProfile = Model->GetSelectedProfile();
@@ -1962,6 +1985,7 @@ FText SProjectLauncherCookByTheBookSettings::HandleHttpChunkInstallNameTextBlock
 	return FText();
 }
 
+
 void SProjectLauncherCookByTheBookSettings::HandleHtppChunkInstallNameCommitted(const FText& NewText, ETextCommit::Type CommitType)
 {
 	if (CommitType == ETextCommit::OnEnter)
@@ -1976,8 +2000,6 @@ void SProjectLauncherCookByTheBookSettings::HandleHtppChunkInstallNameCommitted(
 }
 
 
-
-// Callback for check state changes of the 'UnrealPak' check box.
 void SProjectLauncherCookByTheBookSettings::HandleDLCIncludeEngineContentCheckBoxCheckStateChanged(ECheckBoxState NewState)
 {
 	ILauncherProfilePtr SelectedProfile = Model->GetSelectedProfile();
@@ -1988,7 +2010,7 @@ void SProjectLauncherCookByTheBookSettings::HandleDLCIncludeEngineContentCheckBo
 	}
 }
 
-// Callback for determining the checked state of the 'UnrealPak' check box.
+
 ECheckBoxState SProjectLauncherCookByTheBookSettings::HandleDLCIncludeEngineContentCheckBoxIsChecked() const
 {
 	ILauncherProfilePtr SelectedProfile = Model->GetSelectedProfile();
@@ -2004,6 +2026,7 @@ ECheckBoxState SProjectLauncherCookByTheBookSettings::HandleDLCIncludeEngineCont
 	return ECheckBoxState::Unchecked;
 }
 
+
 void SProjectLauncherCookByTheBookSettings::HandleGenerateChunksCheckBoxCheckStateChanged(ECheckBoxState NewState)
 {
 	ILauncherProfilePtr SelectedProfile = Model->GetSelectedProfile();
@@ -2014,6 +2037,7 @@ void SProjectLauncherCookByTheBookSettings::HandleGenerateChunksCheckBoxCheckSta
 	}
 }
 
+
 ECheckBoxState SProjectLauncherCookByTheBookSettings::HandleGenerateChunksCheckBoxIsChecked() const
 {
 	ILauncherProfilePtr SelectedProfile = Model->GetSelectedProfile();
@@ -2025,9 +2049,6 @@ ECheckBoxState SProjectLauncherCookByTheBookSettings::HandleGenerateChunksCheckB
 }
 
 
-
-
-
 void SProjectLauncherCookByTheBookSettings::HandleDontIncludeEditorContentCheckBoxCheckStateChanged(ECheckBoxState NewState)
 {
 	ILauncherProfilePtr SelectedProfile = Model->GetSelectedProfile();
@@ -2038,6 +2059,7 @@ void SProjectLauncherCookByTheBookSettings::HandleDontIncludeEditorContentCheckB
 	}
 }
 
+
 ECheckBoxState SProjectLauncherCookByTheBookSettings::HandleDontIncludeEditorContentCheckBoxIsChecked() const
 {
 	ILauncherProfilePtr SelectedProfile = Model->GetSelectedProfile();
@@ -2048,4 +2070,5 @@ ECheckBoxState SProjectLauncherCookByTheBookSettings::HandleDontIncludeEditorCon
 	return ECheckBoxState::Unchecked;
 }
 
+
 #undef LOCTEXT_NAMESPACE
