Copying //UE4/Dev-Sequencer to //UE4/Main (Source: //UE4/Dev-Sequencer @ 3617642)

diff --git a/Engine/Source/Developer/ProjectLauncher/Private/Widgets/Cook/SProjectLauncherCookByTheBookSettings.h b/Engine/Source/Developer/ProjectLauncher/Private/Widgets/Cook/SProjectLauncherCookByTheBookSettings.h
--- a/Engine/Source/Developer/ProjectLauncher/Private/Widgets/Cook/SProjectLauncherCookByTheBookSettings.h
+++ b/Engine/Source/Developer/ProjectLauncher/Private/Widgets/Cook/SProjectLauncherCookByTheBookSettings.h
@@ -42,7 +42,7 @@ public:
 public:
 
 	/** Destructor. */
-	~SProjectLauncherCookByTheBookSettings( );
+	~SProjectLauncherCookByTheBookSettings();
 
 public:
 
@@ -52,84 +52,84 @@ public:
 	 * @param InArgs The Slate argument list.
 	 * @param InModel The data model.
 	 */
-	void Construct(	const FArguments& InArgs, const FProjectLauncherModelRef& InModel, bool InShowSimple = false );
+	void Construct(	const FArguments& InArgs, const TSharedRef<FProjectLauncherModel>& InModel, bool InShowSimple = false);
 
 protected:
 
 	/** Refreshes the list of available maps. */
-	void RefreshMapList( );
+	void RefreshMapList();
 
 	/** Refreshes the list of available cultures. */
-	void RefreshCultureList( );
+	void RefreshCultureList();
 
 private:
 
 	/** Handles clicking the 'Select All Cultures' button. */
-	void HandleAllCulturesHyperlinkNavigate( bool AllPlatforms );
+	void HandleAllCulturesHyperlinkNavigate(bool AllPlatforms);
 
 	/** Handles determining the visibility of the 'Select All Cultures' button. */
-	EVisibility HandleAllCulturesHyperlinkVisibility( ) const;
+	EVisibility HandleAllCulturesHyperlinkVisibility() const;
 
 	/** Handles clicking the 'Select All Maps' button. */
-	void HandleAllMapsHyperlinkNavigate( bool AllPlatforms );
+	void HandleAllMapsHyperlinkNavigate(bool AllPlatforms);
 
 	/** Handles selecting a build configuration for the cooker. */
-	void HandleCookConfigurationSelectorConfigurationSelected( EBuildConfigurations::Type );
+	void HandleCookConfigurationSelectorConfigurationSelected(EBuildConfigurations::Type);
 
 	/** Handles getting the content text of the cooker build configuration selector. */
-	FText HandleCookConfigurationSelectorText( ) const;
+	FText HandleCookConfigurationSelectorText() const;
 
 	/** Handles check state changes of the 'Incremental' check box. */
-	void HandleIncrementalCheckBoxCheckStateChanged( ECheckBoxState NewState );
+	void HandleIncrementalCheckBoxCheckStateChanged(ECheckBoxState NewState);
 
 	/** Handles determining the checked state of the 'Incremental' check box. */
-	ECheckBoxState HandleIncrementalCheckBoxIsChecked( ) const;
+	ECheckBoxState HandleIncrementalCheckBoxIsChecked() const;
 	
 	void HandleSharedCookedBuildCheckBoxCheckStateChanged(ECheckBoxState NewState);
 	ECheckBoxState HandleSharedCookedBuildCheckBoxIsChecked() const;
 
 
-	void HandleCompressedCheckBoxCheckStateChanged( ECheckBoxState NewState );
-	void HandleEncryptIniFilesCheckBoxCheckStateChanged(ECheckBoxState NewState );
+	void HandleCompressedCheckBoxCheckStateChanged(ECheckBoxState NewState);
+	void HandleEncryptIniFilesCheckBoxCheckStateChanged(ECheckBoxState NewState);
 
-	ECheckBoxState HandleCompressedCheckBoxIsChecked( ) const;
+	ECheckBoxState HandleCompressedCheckBoxIsChecked() const;
 	ECheckBoxState HandleEncryptIniFilesCheckBoxIsChecked() const;
 	
 	
 	
 	
 	/** Handles generating a row widget in the map list view. */
-	TSharedRef<ITableRow> HandleMapListViewGenerateRow( TSharedPtr<FString> InItem, const TSharedRef<STableViewBase>& OwnerTable );
+	TSharedRef<ITableRow> HandleMapListViewGenerateRow(TSharedPtr<FString> InItem, const TSharedRef<STableViewBase>& OwnerTable);
 
 	/** Handles generating a row widget in the culture list view. */
-	TSharedRef<ITableRow> HandleCultureListViewGenerateRow( TSharedPtr<FString> InItem, const TSharedRef<STableViewBase>& OwnerTable );
+	TSharedRef<ITableRow> HandleCultureListViewGenerateRow(TSharedPtr<FString> InItem, const TSharedRef<STableViewBase>& OwnerTable);
 
 	/** Handles determining the visibility of the 'Select All Maps' button. */
-	EVisibility HandleMapSelectionHyperlinkVisibility( ) const;
+	EVisibility HandleMapSelectionHyperlinkVisibility() const;
 
 	/** Handles getting the visibility of the map selection warning message. */
-	EVisibility HandleNoMapSelectedBoxVisibility( ) const;
+	EVisibility HandleNoMapSelectedBoxVisibility() const;
 
 	/** Handles getting the text in the 'No maps' text block. */
-	FText HandleNoMapsTextBlockText( ) const;
+	FText HandleNoMapsTextBlockText() const;
 
 	/** Handles changing the selected profile in the profile manager. */
-	void HandleProfileManagerProfileSelected( const ILauncherProfilePtr& SelectedProfile, const ILauncherProfilePtr& PreviousProfile );
+	void HandleProfileManagerProfileSelected(const ILauncherProfilePtr& SelectedProfile, const ILauncherProfilePtr& PreviousProfile);
 
 	/** Handles checking the specified 'Show maps' check box. */
-	void HandleShowCheckBoxCheckStateChanged( ECheckBoxState NewState, EShowMapsChoices::Type Choice );
+	void HandleShowCheckBoxCheckStateChanged(ECheckBoxState NewState, EShowMapsChoices::Type Choice);
 
 	/** Handles determining the checked state of the specified 'Show maps' check box. */
-	ECheckBoxState HandleShowCheckBoxIsChecked( EShowMapsChoices::Type Choice ) const;
+	ECheckBoxState HandleShowCheckBoxIsChecked(EShowMapsChoices::Type Choice) const;
 
 	/** Handles checking the specified 'Unversioned' check box. */
-	void HandleUnversionedCheckBoxCheckStateChanged( ECheckBoxState NewState );
+	void HandleUnversionedCheckBoxCheckStateChanged(ECheckBoxState NewState);
 
 	/** Handles determining the checked state of the specified 'Unversioned' check box. */
-	ECheckBoxState HandleUnversionedCheckBoxIsChecked( ) const;
+	ECheckBoxState HandleUnversionedCheckBoxIsChecked() const;
 
 	/** Handles determining the visibility of a validation error icon. */
-	EVisibility HandleValidationErrorIconVisibility( ELauncherProfileValidationErrors::Type Error ) const;
+	EVisibility HandleValidationErrorIconVisibility(ELauncherProfileValidationErrors::Type Error) const;
 
 	/** Callback for getting the cookers additional options. */
 	FText HandleCookOptionsTextBlockText() const;
@@ -145,10 +145,10 @@ private:
 	void HandleProfileProjectChanged();
 
 	// Callback for check state changes of the 'UnrealPak' check box.
-	void HandleUnrealPakCheckBoxCheckStateChanged( ECheckBoxState NewState );
+	void HandleUnrealPakCheckBoxCheckStateChanged(ECheckBoxState NewState);
 
 	// Callback for determining the checked state of the 'UnrealPak' check box.
-	ECheckBoxState HandleUnrealPakCheckBoxIsChecked( ) const;
+	ECheckBoxState HandleUnrealPakCheckBoxIsChecked() const;
 
 	void HandleGenerateChunksCheckBoxCheckStateChanged(ECheckBoxState NewState);
 	ECheckBoxState HandleGenerateChunksCheckBoxIsChecked() const;
@@ -162,9 +162,9 @@ private:
 	//////////////////////////////////////////////////////////////////////////
 	// creating release version related functions
 	// Callback for check state changes of the 'GeneratePatch' check box.
-	void HandleCreateReleaseVersionCheckBoxCheckStateChanged( ECheckBoxState NewState );
+	void HandleCreateReleaseVersionCheckBoxCheckStateChanged(ECheckBoxState NewState);
 	// Callback for determining the checked state of the 'GeneratePatch' check box.
-	ECheckBoxState HandleCreateReleaseVersionCheckBoxIsChecked( ) const;
+	ECheckBoxState HandleCreateReleaseVersionCheckBoxIsChecked() const;
 
 	void HandleCreateReleaseVersionNameCommitted(const FText& NewText, ETextCommit::Type CommitType);
 	FText HandleCreateReleaseVersionNameTextBlockText() const;
@@ -175,9 +175,9 @@ private:
 	//////////////////////////////////////////////////////////////////////////
 	// patch generation related functions
 	// Callback for check state changes of the 'GeneratePatch' check box.
-	void HandleGeneratePatchCheckBoxCheckStateChanged( ECheckBoxState NewState );
+	void HandleGeneratePatchCheckBoxCheckStateChanged(ECheckBoxState NewState);
 	// Callback for determining the checked state of the 'GeneratePatch' check box.
-	ECheckBoxState HandleGeneratePatchCheckBoxIsChecked( ) const;
+	ECheckBoxState HandleGeneratePatchCheckBoxIsChecked() const;
 
 	// Callback for check state changes of the 'AddPatchLevel' check box.
 	void HandleAddPatchLevelCheckBoxCheckStateChanged( ECheckBoxState NewState );
@@ -196,7 +196,7 @@ private:
 
 	//////////////////////////////////////////////////////////////////////////
 	// dlc check box related functions 
-	void HandleBuildDLCCheckBoxCheckStateChanged( ECheckBoxState NewState );
+	void HandleBuildDLCCheckBoxCheckStateChanged(ECheckBoxState NewState);
 	ECheckBoxState HandleBuildDLCCheckBoxIsChecked() const;
 	void HandleDLCNameCommitted(const FText& NewText, ETextCommit::Type CommitType);
 	FText HandleDLCNameTextBlockText() const;
@@ -241,7 +241,7 @@ private:
 	TSharedPtr<SListView<TSharedPtr<FString> > > MapListView;
 
 	/** Holds a pointer to the data model. */
-	FProjectLauncherModelPtr Model;
+	TSharedPtr<FProjectLauncherModel> Model;
 
 	/** Holds the current 'Show maps' check box choice. */
 	EShowMapsChoices::Type ShowMapsChoice;
diff --git a/Engine/Source/Developer/ProjectLauncher/Private/Widgets/Cook/SProjectLauncherCookOnTheFlySettings.cpp b/Engine/Source/Developer/ProjectLauncher/Private/Widgets/Cook/SProjectLauncherCookOnTheFlySettings.cpp
--- a/Engine/Source/Developer/ProjectLauncher/Private/Widgets/Cook/SProjectLauncherCookOnTheFlySettings.cpp
+++ b/Engine/Source/Developer/ProjectLauncher/Private/Widgets/Cook/SProjectLauncherCookOnTheFlySettings.cpp
@@ -1,14 +1,16 @@
 // Copyright 1998-2017 Epic Games, Inc. All Rights Reserved.
 
-#include "Widgets/Cook/SProjectLauncherCookOnTheFlySettings.h"
-#include "Widgets/SBoxPanel.h"
-#include "Styling/SlateTypes.h"
+#include "SProjectLauncherCookOnTheFlySettings.h"
+
 #include "SlateOptMacros.h"
-#include "Widgets/Text/STextBlock.h"
+#include "Styling/SlateTypes.h"
+#include "Widgets/SBoxPanel.h"
 #include "Widgets/Input/SEditableTextBox.h"
 #include "Widgets/Input/SCheckBox.h"
-#include "Widgets/Shared/SProjectLauncherFormLabel.h"
 #include "Widgets/Layout/SExpandableArea.h"
+#include "Widgets/Text/STextBlock.h"
+
+#include "Widgets/Shared/SProjectLauncherFormLabel.h"
 
 
 #define LOCTEXT_NAMESPACE "SProjectLauncherCookOnTheFlySettings"
@@ -17,7 +19,7 @@
 /* SProjectLauncherCookOnTheFlySettings structors
  *****************************************************************************/
 
-SProjectLauncherCookOnTheFlySettings::~SProjectLauncherCookOnTheFlySettings( )
+SProjectLauncherCookOnTheFlySettings::~SProjectLauncherCookOnTheFlySettings()
 {
 	if (Model.IsValid())
 	{
@@ -30,7 +32,7 @@ SProjectLauncherCookOnTheFlySettings::~SProjectLauncherCookOnTheFlySettings( )
  *****************************************************************************/
 
 BEGIN_SLATE_FUNCTION_BUILD_OPTIMIZATION
-void SProjectLauncherCookOnTheFlySettings::Construct( const FArguments& InArgs, const FProjectLauncherModelRef& InModel )
+void SProjectLauncherCookOnTheFlySettings::Construct(const FArguments& InArgs, const TSharedRef<FProjectLauncherModel>& InModel)
 {
 	Model = InModel;
 
@@ -96,7 +98,7 @@ END_SLATE_FUNCTION_BUILD_OPTIMIZATION
 /* SProjectLauncherCookOnTheFlySettings callbacks
  *****************************************************************************/
 
-void SProjectLauncherCookOnTheFlySettings::HandleIncrementalCheckBoxCheckStateChanged( ECheckBoxState NewState )
+void SProjectLauncherCookOnTheFlySettings::HandleIncrementalCheckBoxCheckStateChanged(ECheckBoxState NewState)
 {
 	ILauncherProfilePtr SelectedProfile = Model->GetSelectedProfile();
 
@@ -107,7 +109,7 @@ void SProjectLauncherCookOnTheFlySettings::HandleIncrementalCheckBoxCheckStateCh
 }
 
 
-ECheckBoxState SProjectLauncherCookOnTheFlySettings::HandleIncrementalCheckBoxIsChecked( ) const
+ECheckBoxState SProjectLauncherCookOnTheFlySettings::HandleIncrementalCheckBoxIsChecked() const
 {
 	ILauncherProfilePtr SelectedProfile = Model->GetSelectedProfile();
 
@@ -123,13 +125,13 @@ ECheckBoxState SProjectLauncherCookOnTheFlySettings::HandleIncrementalCheckBoxIs
 }
 
 
-void SProjectLauncherCookOnTheFlySettings::HandleProfileManagerProfileSelected( const ILauncherProfilePtr& SelectedProfile, const ILauncherProfilePtr& PreviousProfile )
+void SProjectLauncherCookOnTheFlySettings::HandleProfileManagerProfileSelected(const ILauncherProfilePtr& SelectedProfile, const ILauncherProfilePtr& PreviousProfile)
 {
 
 }
 
 
-EVisibility SProjectLauncherCookOnTheFlySettings::HandleValidationErrorIconVisibility( ELauncherProfileValidationErrors::Type Error ) const
+EVisibility SProjectLauncherCookOnTheFlySettings::HandleValidationErrorIconVisibility(ELauncherProfileValidationErrors::Type Error) const
 {
 	ILauncherProfilePtr SelectedProfile = Model->GetSelectedProfile();
 
@@ -144,6 +146,7 @@ EVisibility SProjectLauncherCookOnTheFlySettings::HandleValidationErrorIconVisib
 	return EVisibility::Hidden;
 }
 
+
 FText SProjectLauncherCookOnTheFlySettings::HandleCookOptionsTextBlockText() const
 {
 	ILauncherProfilePtr SelectedProfile = Model->GetSelectedProfile();
@@ -158,6 +161,7 @@ FText SProjectLauncherCookOnTheFlySettings::HandleCookOptionsTextBlockText() con
 	return result;
 }
 
+
 void SProjectLauncherCookOnTheFlySettings::HandleCookerOptionsCommitted(const FText& NewText, ETextCommit::Type CommitType)
 {
 	ILauncherProfilePtr SelectedProfile = Model->GetSelectedProfile();
@@ -171,6 +175,7 @@ void SProjectLauncherCookOnTheFlySettings::HandleCookerOptionsCommitted(const FT
 		case ETextCommit::OnCleared:
 			useOptions = TEXT("");
 			break;
+
 		default:
 			break;
 		}
diff --git a/Engine/Source/Developer/ProjectLauncher/Private/Widgets/Cook/SProjectLauncherCookPage.cpp b/Engine/Source/Developer/ProjectLauncher/Private/Widgets/Cook/SProjectLauncherCookPage.cpp
--- a/Engine/Source/Developer/ProjectLauncher/Private/Widgets/Cook/SProjectLauncherCookPage.cpp
+++ b/Engine/Source/Developer/ProjectLauncher/Private/Widgets/Cook/SProjectLauncherCookPage.cpp
@@ -1,12 +1,14 @@
 // Copyright 1998-2017 Epic Games, Inc. All Rights Reserved.
 
-#include "Widgets/Cook/SProjectLauncherCookPage.h"
-#include "Widgets/SBoxPanel.h"
-#include "Textures/SlateIcon.h"
+#include "SProjectLauncherCookPage.h"
+
 #include "Framework/Commands/UIAction.h"
-#include "Widgets/Text/STextBlock.h"
 #include "Framework/MultiBox/MultiBoxBuilder.h"
+#include "Textures/SlateIcon.h"
+#include "Widgets/SBoxPanel.h"
 #include "Widgets/Input/SComboButton.h"
+#include "Widgets/Text/STextBlock.h"
+
 #include "Widgets/Cook/SProjectLauncherCookByTheBookSettings.h"
 #include "Widgets/Cook/SProjectLauncherCookOnTheFlySettings.h"
 
@@ -17,7 +19,7 @@
 /* SProjectLauncherCookPage structors
  *****************************************************************************/
 
-SProjectLauncherCookPage::~SProjectLauncherCookPage( )
+SProjectLauncherCookPage::~SProjectLauncherCookPage()
 {
 	if (Model.IsValid())
 	{
@@ -29,7 +31,7 @@ SProjectLauncherCookPage::~SProjectLauncherCookPage( )
 /* SProjectLauncherCookPage interface
  *****************************************************************************/
 
-void SProjectLauncherCookPage::Construct( const FArguments& InArgs, const FProjectLauncherModelRef& InModel )
+void SProjectLauncherCookPage::Construct(const FArguments& InArgs, const TSharedRef<FProjectLauncherModel>& InModel)
 {
 	Model = InModel;
 
@@ -106,7 +108,7 @@ void SProjectLauncherCookPage::Construct( const FArguments& InArgs, const FProje
 /* SProjectLauncherCookPage callbacks
  *****************************************************************************/
 
-EVisibility SProjectLauncherCookPage::HandleCookByTheBookSettingsVisibility( ) const
+EVisibility SProjectLauncherCookPage::HandleCookByTheBookSettingsVisibility() const
 {
 	ILauncherProfilePtr SelectedProfile = Model->GetSelectedProfile();
 
@@ -122,7 +124,7 @@ EVisibility SProjectLauncherCookPage::HandleCookByTheBookSettingsVisibility( ) c
 }
 
 
-FText SProjectLauncherCookPage::HandleCookModeComboButtonContentText( ) const
+FText SProjectLauncherCookPage::HandleCookModeComboButtonContentText() const
 {
 	ILauncherProfilePtr SelectedProfile = Model->GetSelectedProfile();
 
@@ -152,7 +154,7 @@ FText SProjectLauncherCookPage::HandleCookModeComboButtonContentText( ) const
 }
 
 
-void SProjectLauncherCookPage::HandleCookModeMenuEntryClicked( ELauncherProfileCookModes::Type CookMode )
+void SProjectLauncherCookPage::HandleCookModeMenuEntryClicked(ELauncherProfileCookModes::Type CookMode)
 {
 	ILauncherProfilePtr SelectedProfile = Model->GetSelectedProfile();
 
@@ -163,7 +165,7 @@ void SProjectLauncherCookPage::HandleCookModeMenuEntryClicked( ELauncherProfileC
 }
 
 
-EVisibility SProjectLauncherCookPage::HandleCookOnTheFlySettingsVisibility( ) const
+EVisibility SProjectLauncherCookPage::HandleCookOnTheFlySettingsVisibility() const
 {
 	ILauncherProfilePtr SelectedProfile = Model->GetSelectedProfile();
 
@@ -179,7 +181,7 @@ EVisibility SProjectLauncherCookPage::HandleCookOnTheFlySettingsVisibility( ) co
 }
 
 
-void SProjectLauncherCookPage::HandleProfileManagerProfileSelected( const ILauncherProfilePtr& SelectedProfile, const ILauncherProfilePtr& PreviousProfile )
+void SProjectLauncherCookPage::HandleProfileManagerProfileSelected(const ILauncherProfilePtr& SelectedProfile, const ILauncherProfilePtr& PreviousProfile)
 {
 	// reload settings
 }
diff --git a/Engine/Source/Developer/ProjectLauncher/Private/Widgets/Cook/SProjectLauncherCookedPlatforms.cpp b/Engine/Source/Developer/ProjectLauncher/Private/Widgets/Cook/SProjectLauncherCookedPlatforms.cpp
--- a/Engine/Source/Developer/ProjectLauncher/Private/Widgets/Cook/SProjectLauncherCookedPlatforms.cpp
+++ b/Engine/Source/Developer/ProjectLauncher/Private/Widgets/Cook/SProjectLauncherCookedPlatforms.cpp
@@ -1,15 +1,19 @@
 // Copyright 1998-2017 Epic Games, Inc. All Rights Reserved.
 
-#include "Widgets/Cook/SProjectLauncherCookedPlatforms.h"
+#include "SProjectLauncherCookedPlatforms.h"
+
+#include "Widgets/Input/SHyperlink.h"
 #include "Widgets/Layout/SSeparator.h"
 #include "Widgets/Text/STextBlock.h"
 #include "Widgets/Views/SListView.h"
+
 #include "Widgets/Cook/SProjectLauncherPlatformListRow.h"
-#include "Widgets/Input/SHyperlink.h"
+
 
 #define LOCTEXT_NAMESPACE "SProjectLauncherCookedPlatforms"
 
-void SProjectLauncherCookedPlatforms::Construct(const FArguments& InArgs, const FProjectLauncherModelRef& InModel)
+
+void SProjectLauncherCookedPlatforms::Construct(const FArguments& InArgs, const TSharedRef<FProjectLauncherModel>& InModel)
 {
 	Model = InModel;
 
@@ -20,30 +24,31 @@ void SProjectLauncherCookedPlatforms::Construct(const FArguments& InArgs, const
 			SNew(SVerticalBox)
 
 			+ SVerticalBox::Slot()
-			.FillHeight(1.0f)
-			[
-				// platform menu
-				SAssignNew(PlatformListView, SListView<TSharedPtr<FString> >)
-				.HeaderRow(
-				SNew(SHeaderRow)
-				.Visibility(EVisibility::Collapsed)
-
-				+ SHeaderRow::Column("PlatformName")
-				.DefaultLabel(LOCTEXT("PlatformListPlatformNameColumnHeader", "Platform"))
-				.FillWidth(1.0f)
-				)
-				.ItemHeight(16.0f)
-				.ListItemsSource(&PlatformList)
-				.OnGenerateRow(this, &SProjectLauncherCookedPlatforms::HandlePlatformListViewGenerateRow)
-				.SelectionMode(ESelectionMode::None)
-			]
+				.FillHeight(1.0f)
+				[
+					// platform menu
+					SAssignNew(PlatformListView, SListView<TSharedPtr<FString> >)
+						.HeaderRow
+						(
+							SNew(SHeaderRow)
+								.Visibility(EVisibility::Collapsed)
+
+							+ SHeaderRow::Column("PlatformName")
+								.DefaultLabel(LOCTEXT("PlatformListPlatformNameColumnHeader", "Platform"))
+								.FillWidth(1.0f)
+						)
+						.ItemHeight(16.0f)
+						.ListItemsSource(&PlatformList)
+						.OnGenerateRow(this, &SProjectLauncherCookedPlatforms::HandlePlatformListViewGenerateRow)
+						.SelectionMode(ESelectionMode::None)
+				]
 
 			+ SVerticalBox::Slot()
 				.AutoHeight()
 				.Padding(0.0f, 6.0f, 0.0f, 4.0f)
 				[
 					SNew(SSeparator)
-					.Orientation(Orient_Horizontal)
+						.Orientation(Orient_Horizontal)
 				]
 
 			+ SVerticalBox::Slot()
@@ -52,12 +57,12 @@ void SProjectLauncherCookedPlatforms::Construct(const FArguments& InArgs, const
 					SNew(SHorizontalBox)
 
 					+ SHorizontalBox::Slot()
-					.FillWidth(1.0f)
-					.HAlign(HAlign_Right)
-					[
-						SNew(STextBlock)
-						.Text(LOCTEXT("SelectLabel", "Select:"))
-					]
+						.FillWidth(1.0f)
+						.HAlign(HAlign_Right)
+						[
+							SNew(STextBlock)
+								.Text(LOCTEXT("SelectLabel", "Select:"))
+						]
 
 					+ SHorizontalBox::Slot()
 						.AutoWidth()
@@ -65,10 +70,10 @@ void SProjectLauncherCookedPlatforms::Construct(const FArguments& InArgs, const
 						[
 							// all platforms hyper-link
 							SNew(SHyperlink)
-							.OnNavigate(this, &SProjectLauncherCookedPlatforms::HandleAllPlatformsHyperlinkNavigate, true)
-							.Text(LOCTEXT("AllPlatformsHyperlinkLabel", "All"))
-							.ToolTipText(LOCTEXT("AllPlatformsButtonTooltip", "Select all available platforms."))
-							.Visibility(this, &SProjectLauncherCookedPlatforms::HandleAllPlatformsHyperlinkVisibility)
+								.OnNavigate(this, &SProjectLauncherCookedPlatforms::HandleAllPlatformsHyperlinkNavigate, true)
+								.Text(LOCTEXT("AllPlatformsHyperlinkLabel", "All"))
+								.ToolTipText(LOCTEXT("AllPlatformsButtonTooltip", "Select all available platforms."))
+								.Visibility(this, &SProjectLauncherCookedPlatforms::HandleAllPlatformsHyperlinkVisibility)
 						]
 
 					+ SHorizontalBox::Slot()
@@ -76,15 +81,16 @@ void SProjectLauncherCookedPlatforms::Construct(const FArguments& InArgs, const
 						[
 							// no platforms hyper-link
 							SNew(SHyperlink)
-							.OnNavigate(this, &SProjectLauncherCookedPlatforms::HandleAllPlatformsHyperlinkNavigate, false)
-							.Text(LOCTEXT("NoPlatformsHyperlinkLabel", "None"))
-							.ToolTipText(LOCTEXT("NoPlatformsHyperlinkTooltip", "Deselect all platforms."))
-							.Visibility(this, &SProjectLauncherCookedPlatforms::HandleAllPlatformsHyperlinkVisibility)
+								.OnNavigate(this, &SProjectLauncherCookedPlatforms::HandleAllPlatformsHyperlinkNavigate, false)
+								.Text(LOCTEXT("NoPlatformsHyperlinkLabel", "None"))
+								.ToolTipText(LOCTEXT("NoPlatformsHyperlinkTooltip", "Deselect all platforms."))
+								.Visibility(this, &SProjectLauncherCookedPlatforms::HandleAllPlatformsHyperlinkVisibility)
 						]
 				]
 		];
 }
 
+
 TSharedRef<ITableRow> SProjectLauncherCookedPlatforms::HandlePlatformListViewGenerateRow(TSharedPtr<FString> InItem, const TSharedRef<STableViewBase>& OwnerTable)
 {
 	return SNew(SProjectLauncherPlatformListRow, Model.ToSharedRef())
@@ -92,4 +98,5 @@ TSharedRef<ITableRow> SProjectLauncherCookedPlatforms::HandlePlatformListViewGen
 		.OwnerTableView(OwnerTable);
 }
 
+
 #undef LOCTEXT_NAMESPACE
diff --git a/Engine/Source/Developer/ProjectLauncher/Private/Widgets/Cook/SProjectLauncherCookedPlatforms.h b/Engine/Source/Developer/ProjectLauncher/Private/Widgets/Cook/SProjectLauncherCookedPlatforms.h
--- a/Engine/Source/Developer/ProjectLauncher/Private/Widgets/Cook/SProjectLauncherCookedPlatforms.h
+++ b/Engine/Source/Developer/ProjectLauncher/Private/Widgets/Cook/SProjectLauncherCookedPlatforms.h
@@ -14,6 +14,7 @@
 #include "Interfaces/ITargetPlatform.h"
 #include "Interfaces/ITargetPlatformManagerModule.h"
 
+
 /**
  * Implements the cooked platforms panel.
  */
@@ -28,21 +29,21 @@ public:
 public:
 
 	/**
-	 * Constructs the widget.
+	 * Construct the widget.
 	 *
 	 * @param InArgs The Slate argument list.
 	 * @param InModel The data model.
 	 */
-	void Construct(	const FArguments& InArgs, const FProjectLauncherModelRef& InModel );
+	void Construct(const FArguments& InArgs, const TSharedRef<FProjectLauncherModel>& InModel);
 
 protected:
 
 	/**
-	 * Builds the platform menu.
+	 * Build the platform menu.
 	 *
 	 * @return Platform menu widget.
 	 */
-	void MakePlatformMenu( )
+	void MakePlatformMenu()
 	{
 		TArray<ITargetPlatform*> Platforms = GetTargetPlatformManager()->GetTargetPlatforms();
 
@@ -60,8 +61,8 @@ protected:
 
 private:
 
-	// Callback for clicking the 'Select All Platforms' button.
-	void HandleAllPlatformsHyperlinkNavigate( bool AllPlatforms )
+	/** Callback for clicking the 'Select All Platforms' button. */
+	void HandleAllPlatformsHyperlinkNavigate(bool AllPlatforms)
 	{
 		ILauncherProfilePtr SelectedProfile = Model->GetSelectedProfile();
 
@@ -83,8 +84,8 @@ private:
 		}
 	}
 
-	// Callback for determining the visibility of the 'Select All Platforms' button.
-	EVisibility HandleAllPlatformsHyperlinkVisibility( ) const
+	/** Callback for determining the visibility of the 'Select All Platforms' button. */
+	EVisibility HandleAllPlatformsHyperlinkVisibility() const
 	{
 		if (GetTargetPlatformManager()->GetTargetPlatforms().Num() > 1)
 		{
@@ -94,8 +95,8 @@ private:
 		return EVisibility::Collapsed;
 	}
 
-	// Callback for getting the color of a platform menu check box.
-	FSlateColor HandlePlatformMenuEntryColorAndOpacity( FString PlatformName ) const
+	/** Callback for getting the color of a platform menu check box. */
+	FSlateColor HandlePlatformMenuEntryColorAndOpacity(FString PlatformName) const
 	{
 		ILauncherProfilePtr SelectedProfile = Model->GetSelectedProfile();
 
@@ -115,19 +116,19 @@ private:
 		return FLinearColor::Yellow;
 	}
 
-	// Handles generating a row widget in the map list view.
-	TSharedRef<ITableRow> HandlePlatformListViewGenerateRow( TSharedPtr<FString> InItem, const TSharedRef<STableViewBase>& OwnerTable );
+	/** Callback for generating a row widget in the map list view. */
+	TSharedRef<ITableRow> HandlePlatformListViewGenerateRow(TSharedPtr<FString> InItem, const TSharedRef<STableViewBase>& OwnerTable);
 
 private:
 
-	// Holds a pointer to the data model.
-	FProjectLauncherModelPtr Model;
+	/** Pointer to the data model. */
+	TSharedPtr<FProjectLauncherModel> Model;
 
-	// Holds the platform list.
+	/** The platform list. */
 	TArray<TSharedPtr<FString> > PlatformList;
 
-	// Holds the platform list view.
-	TSharedPtr<SListView<TSharedPtr<FString> > > PlatformListView;
+	/** The platform list view. */
+	TSharedPtr<SListView<TSharedPtr<FString>>> PlatformListView;
 
 };
 
diff --git a/Engine/Source/Developer/ProjectLauncher/Private/Widgets/Cook/SProjectLauncherCultureListRow.h b/Engine/Source/Developer/ProjectLauncher/Private/Widgets/Cook/SProjectLauncherCultureListRow.h
--- a/Engine/Source/Developer/ProjectLauncher/Private/Widgets/Cook/SProjectLauncherCultureListRow.h
+++ b/Engine/Source/Developer/ProjectLauncher/Private/Widgets/Cook/SProjectLauncherCultureListRow.h
@@ -39,7 +39,7 @@ public:
 	 * @param InArgs The construction arguments.
 	 * @param InProfileManager The profile manager to use.
 	 */
-	void Construct( const FArguments& InArgs, const FProjectLauncherModelRef& InModel )
+	void Construct( const FArguments& InArgs, const TSharedRef<FProjectLauncherModel>& InModel )
 	{
 		HighlightString = InArgs._HighlightString;
 		CultureName = InArgs._CultureName;
@@ -116,7 +116,7 @@ private:
 	TSharedPtr<FString> CultureName;
 
 	// Holds a pointer to the data model.
-	FProjectLauncherModelPtr Model;
+	TSharedPtr<FProjectLauncherModel> Model;
 };
 
 
diff --git a/Engine/Source/Developer/ProjectLauncher/Private/Widgets/Cook/SProjectLauncherMapListRow.h b/Engine/Source/Developer/ProjectLauncher/Private/Widgets/Cook/SProjectLauncherMapListRow.h
--- a/Engine/Source/Developer/ProjectLauncher/Private/Widgets/Cook/SProjectLauncherMapListRow.h
+++ b/Engine/Source/Developer/ProjectLauncher/Private/Widgets/Cook/SProjectLauncherMapListRow.h
@@ -39,7 +39,7 @@ public:
 	 * @param InArgs The construction arguments.
 	 * @param InProfileManager The profile manager to use.
 	 */
-	void Construct( const FArguments& InArgs, const FProjectLauncherModelRef& InModel )
+	void Construct( const FArguments& InArgs, const TSharedRef<FProjectLauncherModel>& InModel )
 	{
 		HighlightString = InArgs._HighlightString;
 		MapName = InArgs._MapName;
@@ -115,7 +115,7 @@ private:
 	TSharedPtr<FString> MapName;
 
 	// Holds a pointer to the data model.
-	FProjectLauncherModelPtr Model;
+	TSharedPtr<FProjectLauncherModel> Model;
 };
 
 
diff --git a/Engine/Source/Developer/ProjectLauncher/Private/Widgets/Cook/SProjectLauncherSimpleCookPage.cpp b/Engine/Source/Developer/ProjectLauncher/Private/Widgets/Cook/SProjectLauncherSimpleCookPage.cpp
--- a/Engine/Source/Developer/ProjectLauncher/Private/Widgets/Cook/SProjectLauncherSimpleCookPage.cpp
+++ b/Engine/Source/Developer/ProjectLauncher/Private/Widgets/Cook/SProjectLauncherSimpleCookPage.cpp
@@ -1,6 +1,7 @@
 // Copyright 1998-2017 Epic Games, Inc. All Rights Reserved.
 
-#include "Widgets/Cook/SProjectLauncherSimpleCookPage.h"
+#include "SProjectLauncherSimpleCookPage.h"
+
 #include "Widgets/SBoxPanel.h"
 #include "Widgets/Cook/SProjectLauncherCookByTheBookSettings.h"
 
@@ -11,7 +12,7 @@
 /* SProjectLauncherSimpleCookPage structors
  *****************************************************************************/
 
-SProjectLauncherSimpleCookPage::~SProjectLauncherSimpleCookPage( )
+SProjectLauncherSimpleCookPage::~SProjectLauncherSimpleCookPage()
 {
 	if (Model.IsValid())
 	{
@@ -23,7 +24,7 @@ SProjectLauncherSimpleCookPage::~SProjectLauncherSimpleCookPage( )
 /* SProjectLauncherSimpleCookPage interface
  *****************************************************************************/
 
-void SProjectLauncherSimpleCookPage::Construct( const FArguments& InArgs, const FProjectLauncherModelRef& InModel )
+void SProjectLauncherSimpleCookPage::Construct(const FArguments& InArgs, const TSharedRef<FProjectLauncherModel>& InModel)
 {
 	Model = InModel;
 
@@ -46,7 +47,7 @@ void SProjectLauncherSimpleCookPage::Construct( const FArguments& InArgs, const
 /* SProjectLauncherSimpleCookPage callbacks
  *****************************************************************************/
 
-EVisibility SProjectLauncherSimpleCookPage::HandleCookByTheBookSettingsVisibility( ) const
+EVisibility SProjectLauncherSimpleCookPage::HandleCookByTheBookSettingsVisibility() const
 {
 	ILauncherProfilePtr SelectedProfile = Model->GetSelectedProfile();
 
@@ -62,7 +63,7 @@ EVisibility SProjectLauncherSimpleCookPage::HandleCookByTheBookSettingsVisibilit
 }
 
 
-FText SProjectLauncherSimpleCookPage::HandleCookModeComboButtonContentText( ) const
+FText SProjectLauncherSimpleCookPage::HandleCookModeComboButtonContentText() const
 {
 	ILauncherProfilePtr SelectedProfile = Model->GetSelectedProfile();
 
@@ -92,7 +93,7 @@ FText SProjectLauncherSimpleCookPage::HandleCookModeComboButtonContentText( ) co
 }
 
 
-void SProjectLauncherSimpleCookPage::HandleCookModeMenuEntryClicked( ELauncherProfileCookModes::Type CookMode )
+void SProjectLauncherSimpleCookPage::HandleCookModeMenuEntryClicked(ELauncherProfileCookModes::Type CookMode)
 {
 	ILauncherProfilePtr SelectedProfile = Model->GetSelectedProfile();
 
@@ -103,7 +104,7 @@ void SProjectLauncherSimpleCookPage::HandleCookModeMenuEntryClicked( ELauncherPr
 }
 
 
-EVisibility SProjectLauncherSimpleCookPage::HandleCookOnTheFlySettingsVisibility( ) const
+EVisibility SProjectLauncherSimpleCookPage::HandleCookOnTheFlySettingsVisibility() const
 {
 	ILauncherProfilePtr SelectedProfile = Model->GetSelectedProfile();
 
@@ -119,7 +120,7 @@ EVisibility SProjectLauncherSimpleCookPage::HandleCookOnTheFlySettingsVisibility
 }
 
 
-void SProjectLauncherSimpleCookPage::HandleProfileManagerProfileSelected( const ILauncherProfilePtr& SelectedProfile, const ILauncherProfilePtr& PreviousProfile )
+void SProjectLauncherSimpleCookPage::HandleProfileManagerProfileSelected(const ILauncherProfilePtr& SelectedProfile, const ILauncherProfilePtr& PreviousProfile)
 {
 	// reload settings
 }
diff --git a/Engine/Source/Developer/ProjectLauncher/Private/Widgets/Deploy/SProjectLauncherDeployFileServerSettings.cpp b/Engine/Source/Developer/ProjectLauncher/Private/Widgets/Deploy/SProjectLauncherDeployFileServerSettings.cpp
--- a/Engine/Source/Developer/ProjectLauncher/Private/Widgets/Deploy/SProjectLauncherDeployFileServerSettings.cpp
+++ b/Engine/Source/Developer/ProjectLauncher/Private/Widgets/Deploy/SProjectLauncherDeployFileServerSettings.cpp
@@ -1,12 +1,15 @@
 // Copyright 1998-2017 Epic Games, Inc. All Rights Reserved.
 
-#include "Widgets/Deploy/SProjectLauncherDeployFileServerSettings.h"
-#include "Widgets/SBoxPanel.h"
+#include "SProjectLauncherDeployFileServerSettings.h"
+
+#include "EditorStyleSet.h"
 #include "Styling/SlateTypes.h"
+#include "Widgets/SBoxPanel.h"
+#include "Widgets/Input/SCheckBox.h"
 #include "Widgets/Layout/SBorder.h"
 #include "Widgets/Text/STextBlock.h"
-#include "Widgets/Input/SCheckBox.h"
-#include "EditorStyleSet.h"
+
+#include "Models/ProjectLauncherModel.h"
 #include "Widgets/Deploy/SProjectLauncherDeployTargets.h"
 #include "Widgets/Layout/SExpandableArea.h"
 
@@ -17,7 +20,7 @@
 /* SProjectLauncherDeployFileServerSettings interface
  *****************************************************************************/
 
-void SProjectLauncherDeployFileServerSettings::Construct( const FArguments& InArgs, const FProjectLauncherModelRef& InModel )
+void SProjectLauncherDeployFileServerSettings::Construct(const FArguments& InArgs, const TSharedRef<FProjectLauncherModel>& InModel)
 {
 	Model = InModel;
 
@@ -90,7 +93,7 @@ void SProjectLauncherDeployFileServerSettings::Construct( const FArguments& InAr
 /* SProjectLauncherDeployFileServerSettings callbacks
  *****************************************************************************/
 
-void SProjectLauncherDeployFileServerSettings::HandleHideWindowCheckBoxCheckStateChanged( ECheckBoxState NewState )
+void SProjectLauncherDeployFileServerSettings::HandleHideWindowCheckBoxCheckStateChanged(ECheckBoxState NewState)
 {
 	ILauncherProfilePtr SelectedProfile = Model->GetSelectedProfile();
 
@@ -101,7 +104,7 @@ void SProjectLauncherDeployFileServerSettings::HandleHideWindowCheckBoxCheckStat
 }
 
 
-ECheckBoxState SProjectLauncherDeployFileServerSettings::HandleHideWindowCheckBoxIsChecked( ) const
+ECheckBoxState SProjectLauncherDeployFileServerSettings::HandleHideWindowCheckBoxIsChecked() const
 {
 	ILauncherProfilePtr SelectedProfile = Model->GetSelectedProfile();
 
@@ -117,7 +120,7 @@ ECheckBoxState SProjectLauncherDeployFileServerSettings::HandleHideWindowCheckBo
 }
 
 
-void SProjectLauncherDeployFileServerSettings::HandleStreamingServerCheckBoxCheckStateChanged( ECheckBoxState NewState )
+void SProjectLauncherDeployFileServerSettings::HandleStreamingServerCheckBoxCheckStateChanged(ECheckBoxState NewState)
 {
 	ILauncherProfilePtr SelectedProfile = Model->GetSelectedProfile();
 
@@ -128,7 +131,7 @@ void SProjectLauncherDeployFileServerSettings::HandleStreamingServerCheckBoxChec
 }
 
 
-ECheckBoxState SProjectLauncherDeployFileServerSettings::HandleStreamingServerCheckBoxIsChecked( ) const
+ECheckBoxState SProjectLauncherDeployFileServerSettings::HandleStreamingServerCheckBoxIsChecked() const
 {
 	ILauncherProfilePtr SelectedProfile = Model->GetSelectedProfile();
 
diff --git a/Engine/Source/Developer/ProjectLauncher/Private/Widgets/Deploy/SProjectLauncherDeployFileServerSettings.h b/Engine/Source/Developer/ProjectLauncher/Private/Widgets/Deploy/SProjectLauncherDeployFileServerSettings.h
--- a/Engine/Source/Developer/ProjectLauncher/Private/Widgets/Deploy/SProjectLauncherDeployFileServerSettings.h
+++ b/Engine/Source/Developer/ProjectLauncher/Private/Widgets/Deploy/SProjectLauncherDeployFileServerSettings.h
@@ -2,13 +2,15 @@
 
 #pragma once
 
-#include "CoreMinimal.h"
+#include "Templates/SharedPointer.h"
 #include "Widgets/DeclarativeSyntaxSupport.h"
 #include "Widgets/SCompoundWidget.h"
-#include "Models/ProjectLauncherModel.h"
 
+class FArguments;
+class FProjectLauncherModel;
 enum class ECheckBoxState : uint8;
 
+
 /**
  * Implements the deploy-to-device settings panel.
  */
@@ -28,25 +30,24 @@ public:
 	 * @param InArgs The Slate argument list.
 	 * @param InModel The data model.
 	 */
-	void Construct(	const FArguments& InArgs, const FProjectLauncherModelRef& InModel );
-
+	void Construct(const FArguments& InArgs, const TSharedRef<FProjectLauncherModel>& InModel);
 
 private:
 
-	// Callback for check state changes of the 'Hide Window' check box.
-	void HandleHideWindowCheckBoxCheckStateChanged( ECheckBoxState NewState );
+	/** Callback for check state changes of the 'Hide Window' check box. */
+	void HandleHideWindowCheckBoxCheckStateChanged(ECheckBoxState NewState);
 
-	// Callback for determining the checked state of the 'Hide Window' check box.
-	ECheckBoxState HandleHideWindowCheckBoxIsChecked( ) const;
+	/** Callback for determining the checked state of the 'Hide Window' check box. */
+	ECheckBoxState HandleHideWindowCheckBoxIsChecked() const;
 
-	// Callback for check state changes of the 'Streaming Server' check box.
-	void HandleStreamingServerCheckBoxCheckStateChanged( ECheckBoxState NewState );
+	/** Callback for check state changes of the 'Streaming Server' check box. */
+	void HandleStreamingServerCheckBoxCheckStateChanged(ECheckBoxState NewState);
 
-	// Callback for determining the checked state of the 'Streaming Server' check box.
-	ECheckBoxState HandleStreamingServerCheckBoxIsChecked( ) const;
+	/** Callback for determining the checked state of the 'Streaming Server' check box. */
+	ECheckBoxState HandleStreamingServerCheckBoxIsChecked() const;
 
 private:
 
-	// Holds a pointer to the data model.
-	FProjectLauncherModelPtr Model;
+	/** Holds a pointer to the data model. */
+	TSharedPtr<FProjectLauncherModel> Model;
 };
diff --git a/Engine/Source/Developer/ProjectLauncher/Private/Widgets/Deploy/SProjectLauncherDeployPage.cpp b/Engine/Source/Developer/ProjectLauncher/Private/Widgets/Deploy/SProjectLauncherDeployPage.cpp
--- a/Engine/Source/Developer/ProjectLauncher/Private/Widgets/Deploy/SProjectLauncherDeployPage.cpp
+++ b/Engine/Source/Developer/ProjectLauncher/Private/Widgets/Deploy/SProjectLauncherDeployPage.cpp
@@ -1,16 +1,18 @@
 // Copyright 1998-2017 Epic Games, Inc. All Rights Reserved.
 
-#include "Widgets/Deploy/SProjectLauncherDeployPage.h"
-#include "Widgets/SBoxPanel.h"
+#include "SProjectLauncherDeployPage.h"
+
+#include "EditorStyleSet.h"
+#include "Framework/Commands/UIAction.h"
+#include "Framework/MultiBox/MultiBoxBuilder.h"
 #include "SlateOptMacros.h"
 #include "Textures/SlateIcon.h"
-#include "Framework/Commands/UIAction.h"
-#include "Widgets/Layout/SBorder.h"
+#include "Widgets/SBoxPanel.h"
 #include "Widgets/Images/SImage.h"
-#include "Widgets/Text/STextBlock.h"
-#include "Framework/MultiBox/MultiBoxBuilder.h"
 #include "Widgets/Input/SComboButton.h"
-#include "EditorStyleSet.h"
+#include "Widgets/Layout/SBorder.h"
+#include "Widgets/Text/STextBlock.h"
+
 #include "Widgets/Deploy/SProjectLauncherDeployFileServerSettings.h"
 #include "Widgets/Deploy/SProjectLauncherDeployToDeviceSettings.h"
 #include "Widgets/Deploy/SProjectLauncherDeployRepositorySettings.h"
@@ -22,7 +24,7 @@
 /* SProjectLauncherProfilePage structors
  *****************************************************************************/
 
-SProjectLauncherDeployPage::~SProjectLauncherDeployPage( )
+SProjectLauncherDeployPage::~SProjectLauncherDeployPage()
 {
 	if (Model.IsValid())
 	{
@@ -35,7 +37,7 @@ SProjectLauncherDeployPage::~SProjectLauncherDeployPage( )
  *****************************************************************************/
 
 BEGIN_SLATE_FUNCTION_BUILD_OPTIMIZATION
-void SProjectLauncherDeployPage::Construct( const FArguments& InArgs, const FProjectLauncherModelRef& InModel, bool IsFromRepository )
+void SProjectLauncherDeployPage::Construct(const FArguments& InArgs, const TSharedRef<FProjectLauncherModel>& InModel, bool IsFromRepository)
 {
 	Model = InModel;
 
@@ -157,7 +159,7 @@ END_SLATE_FUNCTION_BUILD_OPTIMIZATION
 /* SProjectLauncherDevicesPage callbacks
  *****************************************************************************/
 
-FText SProjectLauncherDeployPage::HandleDeploymentModeComboButtonContentText( ) const
+FText SProjectLauncherDeployPage::HandleDeploymentModeComboButtonContentText() const
 {
 	ILauncherProfilePtr SelectedProfile = Model->GetSelectedProfile();
 
@@ -192,7 +194,7 @@ FText SProjectLauncherDeployPage::HandleDeploymentModeComboButtonContentText( )
 }
 
 
-void SProjectLauncherDeployPage::HandleDeploymentModeMenuEntryClicked( ELauncherProfileDeploymentModes::Type DeploymentMode )
+void SProjectLauncherDeployPage::HandleDeploymentModeMenuEntryClicked(ELauncherProfileDeploymentModes::Type DeploymentMode)
 {
 	ILauncherProfilePtr SelectedProfile = Model->GetSelectedProfile();
 
@@ -203,7 +205,7 @@ void SProjectLauncherDeployPage::HandleDeploymentModeMenuEntryClicked( ELauncher
 }
 
 
-EVisibility SProjectLauncherDeployPage::HandleDeployFileServerSettingsVisibility( ) const
+EVisibility SProjectLauncherDeployPage::HandleDeployFileServerSettingsVisibility() const
 {
 	ILauncherProfilePtr SelectedProfile = Model->GetSelectedProfile();
 
@@ -219,7 +221,7 @@ EVisibility SProjectLauncherDeployPage::HandleDeployFileServerSettingsVisibility
 }
 
 
-EVisibility SProjectLauncherDeployPage::HandleDeployToDeviceSettingsVisibility( ) const
+EVisibility SProjectLauncherDeployPage::HandleDeployToDeviceSettingsVisibility() const
 {
 	ILauncherProfilePtr SelectedProfile = Model->GetSelectedProfile();
 
@@ -238,7 +240,7 @@ EVisibility SProjectLauncherDeployPage::HandleDeployToDeviceSettingsVisibility(
 }
 
 
-EVisibility SProjectLauncherDeployPage::HandleDeployRepositorySettingsVisibility( ) const
+EVisibility SProjectLauncherDeployPage::HandleDeployRepositorySettingsVisibility() const
 {
 	ILauncherProfilePtr SelectedProfile = Model->GetSelectedProfile();
 
@@ -254,13 +256,13 @@ EVisibility SProjectLauncherDeployPage::HandleDeployRepositorySettingsVisibility
 }
 
 
-void SProjectLauncherDeployPage::HandleProfileManagerProfileSelected( const ILauncherProfilePtr& SelectedProfile, const ILauncherProfilePtr& PreviousProfile )
+void SProjectLauncherDeployPage::HandleProfileManagerProfileSelected(const ILauncherProfilePtr& SelectedProfile, const ILauncherProfilePtr& PreviousProfile)
 {
 
 }
 
 
-EVisibility SProjectLauncherDeployPage::HandleValidationErrorIconVisibility( ELauncherProfileValidationErrors::Type Error ) const
+EVisibility SProjectLauncherDeployPage::HandleValidationErrorIconVisibility(ELauncherProfileValidationErrors::Type Error) const
 {
 	ILauncherProfilePtr SelectedProfile = Model->GetSelectedProfile();
 
diff --git a/Engine/Source/Developer/ProjectLauncher/Private/Widgets/Deploy/SProjectLauncherDeployPage.h b/Engine/Source/Developer/ProjectLauncher/Private/Widgets/Deploy/SProjectLauncherDeployPage.h
--- a/Engine/Source/Developer/ProjectLauncher/Private/Widgets/Deploy/SProjectLauncherDeployPage.h
+++ b/Engine/Source/Developer/ProjectLauncher/Private/Widgets/Deploy/SProjectLauncherDeployPage.h
@@ -36,7 +36,7 @@ public:
 	 * @param InArgs - The Slate argument list.
 	 * @param InModel - The data model.
 	 */
-	void Construct(	const FArguments& InArgs, const FProjectLauncherModelRef& InModel, bool IsFromRepository = false );
+	void Construct(	const FArguments& InArgs, const TSharedRef<FProjectLauncherModel>& InModel, bool IsFromRepository = false );
 
 private:
 
@@ -64,5 +64,5 @@ private:
 private:
 
 	// Holds a pointer to the data model.
-	FProjectLauncherModelPtr Model;
+	TSharedPtr<FProjectLauncherModel> Model;
 };
diff --git a/Engine/Source/Developer/ProjectLauncher/Private/Widgets/Deploy/SProjectLauncherDeployRepositorySettings.cpp b/Engine/Source/Developer/ProjectLauncher/Private/Widgets/Deploy/SProjectLauncherDeployRepositorySettings.cpp
--- a/Engine/Source/Developer/ProjectLauncher/Private/Widgets/Deploy/SProjectLauncherDeployRepositorySettings.cpp
+++ b/Engine/Source/Developer/ProjectLauncher/Private/Widgets/Deploy/SProjectLauncherDeployRepositorySettings.cpp
@@ -1,14 +1,16 @@
 // Copyright 1998-2017 Epic Games, Inc. All Rights Reserved.
 
-#include "Widgets/Deploy/SProjectLauncherDeployRepositorySettings.h"
-#include "Widgets/SBoxPanel.h"
+#include "SProjectLauncherDeployRepositorySettings.h"
+
+#include "EditorStyleSet.h"
+#include "DesktopPlatformModule.h"
 #include "Framework/Application/SlateApplication.h"
-#include "Widgets/Layout/SBorder.h"
-#include "Widgets/Text/STextBlock.h"
+#include "Widgets/SBoxPanel.h"
 #include "Widgets/Input/SEditableTextBox.h"
 #include "Widgets/Input/SButton.h"
-#include "EditorStyleSet.h"
-#include "DesktopPlatformModule.h"
+#include "Widgets/Layout/SBorder.h"
+#include "Widgets/Text/STextBlock.h"
+
 #include "Widgets/Deploy/SProjectLauncherDeployTargets.h"
 
 
@@ -18,7 +20,7 @@
 /* SProjectLauncherDeployRepositorySettings interface
  *****************************************************************************/
 
-void SProjectLauncherDeployRepositorySettings::Construct( const FArguments& InArgs, const FProjectLauncherModelRef& InModel )
+void SProjectLauncherDeployRepositorySettings::Construct(const FArguments& InArgs, const TSharedRef<FProjectLauncherModel>& InModel)
 {
 	Model = InModel;
 
@@ -26,54 +28,53 @@ void SProjectLauncherDeployRepositorySettings::Construct( const FArguments& InAr
 	[
 		SNew(SVerticalBox)
 
-
 		+ SVerticalBox::Slot()
 		.FillHeight(1.0)
 		[
 			SNew(SBorder)
-			.Padding(8.0)
-			.BorderImage(FEditorStyle::GetBrush("ToolPanel.GroupBorder"))
-			[
-				SNew(SVerticalBox)
-
-				+ SVerticalBox::Slot()
-				.AutoHeight()
+				.Padding(8.0)
+				.BorderImage(FEditorStyle::GetBrush("ToolPanel.GroupBorder"))
 				[
-					SNew(STextBlock)
-					.Text(LOCTEXT("RepositoryPathLabel", "Repository Path:"))
-				]
-
-				+ SVerticalBox::Slot()
-					.AutoHeight()
-					.Padding(0.0, 4.0, 0.0, 0.0)
-					[
-						SNew(SHorizontalBox)
+					SNew(SVerticalBox)
 
-						+ SHorizontalBox::Slot()
-						.FillWidth(1.0)
-						.Padding(0.0, 0.0, 0.0, 3.0)
+					+ SVerticalBox::Slot()
+						.AutoHeight()
 						[
-							// repository path text box
-							SAssignNew(RepositoryPathTextBox, SEditableTextBox)
-							.OnTextCommitted(this, &SProjectLauncherDeployRepositorySettings::OnTextCommitted)
-							.OnTextChanged(this, &SProjectLauncherDeployRepositorySettings::OnTextChanged)
+							SNew(STextBlock)
+								.Text(LOCTEXT("RepositoryPathLabel", "Repository Path:"))
 						]
 
-						+ SHorizontalBox::Slot()
-							.AutoWidth()
-							.HAlign(HAlign_Right)
-							.Padding(4.0, 0.0, 0.0, 0.0)
-							[
-								// browse button
-								SNew(SButton)
-								.ContentPadding(FMargin(6.0, 2.0))
-								.IsEnabled(true)
-								.Text(LOCTEXT("BrowseButtonText", "Browse..."))
-								.ToolTipText(LOCTEXT("BrowseButtonToolTip", "Browse for the repository"))
-								.OnClicked(this, &SProjectLauncherDeployRepositorySettings::HandleBrowseButtonClicked)
-							]
-					]
-			]
+					+ SVerticalBox::Slot()
+						.AutoHeight()
+						.Padding(0.0, 4.0, 0.0, 0.0)
+						[
+							SNew(SHorizontalBox)
+
+							+ SHorizontalBox::Slot()
+								.FillWidth(1.0)
+								.Padding(0.0, 0.0, 0.0, 3.0)
+								[
+									// repository path text box
+									SAssignNew(RepositoryPathTextBox, SEditableTextBox)
+										.OnTextCommitted(this, &SProjectLauncherDeployRepositorySettings::OnTextCommitted)
+										.OnTextChanged(this, &SProjectLauncherDeployRepositorySettings::OnTextChanged)
+								]
+
+							+ SHorizontalBox::Slot()
+								.AutoWidth()
+								.HAlign(HAlign_Right)
+								.Padding(4.0, 0.0, 0.0, 0.0)
+								[
+									// browse button
+									SNew(SButton)
+										.ContentPadding(FMargin(6.0, 2.0))
+										.IsEnabled(true)
+										.Text(LOCTEXT("BrowseButtonText", "Browse..."))
+										.ToolTipText(LOCTEXT("BrowseButtonToolTip", "Browse for the repository"))
+										.OnClicked(this, &SProjectLauncherDeployRepositorySettings::HandleBrowseButtonClicked)
+								]
+						]
+				]
 		]
 
 		+ SVerticalBox::Slot()
@@ -94,10 +95,10 @@ void SProjectLauncherDeployRepositorySettings::Construct( const FArguments& InAr
 /* SProjectLauncherDeployRepositorySettings callbacks
  *****************************************************************************/
 
-FReply SProjectLauncherDeployRepositorySettings::HandleBrowseButtonClicked( )
+FReply SProjectLauncherDeployRepositorySettings::HandleBrowseButtonClicked()
 {
 	IDesktopPlatform* DesktopPlatform = FDesktopPlatformModule::Get();
-	if ( DesktopPlatform )
+	if (DesktopPlatform)
 	{
 		TSharedPtr<SWindow> ParentWindow = FSlateApplication::Get().FindWidgetWindow(AsShared());
 		void* ParentWindowHandle = (ParentWindow.IsValid() && ParentWindow->GetNativeWindow().IsValid()) ? ParentWindow->GetNativeWindow()->GetOSWindowHandle() : nullptr;
@@ -110,9 +111,9 @@ FReply SProjectLauncherDeployRepositorySettings::HandleBrowseButtonClicked( )
 			FolderName
 			);
 
-		if ( bFolderSelected )
+		if (bFolderSelected)
 		{
-			if ( !FolderName.EndsWith(TEXT("/")) )
+			if (!FolderName.EndsWith(TEXT("/")))
 			{
 				FolderName += TEXT("/");
 			}
@@ -142,7 +143,7 @@ void SProjectLauncherDeployRepositorySettings::OnTextChanged(const FText& InText
 }
 
 
-void SProjectLauncherDeployRepositorySettings::OnTextCommitted( const FText& InText, ETextCommit::Type CommitInfo)
+void SProjectLauncherDeployRepositorySettings::OnTextCommitted(const FText& InText, ETextCommit::Type CommitInfo)
 {
 	if (CommitInfo == ETextCommit::OnEnter)
 	{
diff --git a/Engine/Source/Developer/ProjectLauncher/Private/Widgets/Deploy/SProjectLauncherDeployTargetListRow.h b/Engine/Source/Developer/ProjectLauncher/Private/Widgets/Deploy/SProjectLauncherDeployTargetListRow.h
--- a/Engine/Source/Developer/ProjectLauncher/Private/Widgets/Deploy/SProjectLauncherDeployTargetListRow.h
+++ b/Engine/Source/Developer/ProjectLauncher/Private/Widgets/Deploy/SProjectLauncherDeployTargetListRow.h
@@ -3,12 +3,13 @@
 #pragma once
 
 #include "CoreMinimal.h"
+#include "ILauncherDeviceGroup.h"
+#include "ITargetDeviceProxy.h"
 #include "Misc/Attribute.h"
 #include "Widgets/SNullWidget.h"
 #include "Widgets/DeclarativeSyntaxSupport.h"
 #include "Widgets/SWidget.h"
 #include "Layout/Margin.h"
-#include "Interfaces/ILauncherDeviceGroup.h"
 #include "Widgets/Views/STableViewBase.h"
 #include "Styling/SlateTypes.h"
 #include "Styling/StyleDefaults.h"
@@ -29,7 +30,7 @@
  * Implements a row widget for the launcher's device proxy list.
  */
 class SProjectLauncherDeployTargetListRow
-	: public SMultiColumnTableRow<ITargetDeviceProxyPtr>
+	: public SMultiColumnTableRow<TSharedPtr<ITargetDeviceProxy>>
 {
 public:
 
@@ -43,7 +44,7 @@ public:
 		/**
 		 * The device proxy shown in this row.
 		 */
-		SLATE_ARGUMENT(ITargetDeviceProxyPtr, DeviceProxy)
+		SLATE_ARGUMENT(TSharedPtr<ITargetDeviceProxy>, DeviceProxy)
 
 		/**
 		 * The row's highlight text.
@@ -84,7 +85,7 @@ public:
 		}
 
 
-		SMultiColumnTableRow<ITargetDeviceProxyPtr>::Construct(FSuperRowType::FArguments(), InOwnerTableView);
+		SMultiColumnTableRow<TSharedPtr<ITargetDeviceProxy>>::Construct(FSuperRowType::FArguments(), InOwnerTableView);
 	}
 
 public:
@@ -326,7 +327,7 @@ private:
 	TAttribute<ILauncherDeviceGroupPtr> DeviceGroup;
 
 	// Holds a reference to the device proxy that is displayed in this row.
-	ITargetDeviceProxyPtr DeviceProxy;
+	TSharedPtr<ITargetDeviceProxy> DeviceProxy;
 
 	TSharedPtr<SCheckBox> DeviceCheckbox;
 
diff --git a/Engine/Source/Developer/ProjectLauncher/Private/Widgets/Deploy/SProjectLauncherDeployTargets.cpp b/Engine/Source/Developer/ProjectLauncher/Private/Widgets/Deploy/SProjectLauncherDeployTargets.cpp
--- a/Engine/Source/Developer/ProjectLauncher/Private/Widgets/Deploy/SProjectLauncherDeployTargets.cpp
+++ b/Engine/Source/Developer/ProjectLauncher/Private/Widgets/Deploy/SProjectLauncherDeployTargets.cpp
@@ -1,11 +1,14 @@
 // Copyright 1998-2017 Epic Games, Inc. All Rights Reserved.
 
-#include "Widgets/Deploy/SProjectLauncherDeployTargets.h"
+#include "SProjectLauncherDeployTargets.h"
+
+#include "EditorStyleSet.h"
+#include "ITargetDeviceProxyManager.h"
 #include "Interfaces/ITargetPlatform.h"
 #include "Interfaces/ITargetPlatformManagerModule.h"
 #include "SlateOptMacros.h"
 #include "Widgets/Images/SImage.h"
-#include "EditorStyleSet.h"
+
 #include "Widgets/Deploy/SProjectLauncherDeployTargetListRow.h"
 
 
@@ -19,7 +22,7 @@ SProjectLauncherDeployTargets::~SProjectLauncherDeployTargets()
 {
 	if (Model.IsValid())
 	{
-		const ITargetDeviceProxyManagerRef& DeviceProxyManager = Model->GetDeviceProxyManager();
+		const TSharedRef<ITargetDeviceProxyManager>& DeviceProxyManager = Model->GetDeviceProxyManager();
 		DeviceProxyManager->OnProxyAdded().RemoveAll(this);
 		DeviceProxyManager->OnProxyRemoved().RemoveAll(this);
 	}
@@ -30,7 +33,7 @@ SProjectLauncherDeployTargets::~SProjectLauncherDeployTargets()
  *****************************************************************************/
 
 BEGIN_SLATE_FUNCTION_BUILD_OPTIMIZATION
-void SProjectLauncherDeployTargets::Construct(const FArguments& InArgs, const FProjectLauncherModelRef& InModel)
+void SProjectLauncherDeployTargets::Construct(const FArguments& InArgs, const TSharedRef<FProjectLauncherModel>& InModel)
 {
 	Model = InModel; 
 
@@ -39,94 +42,94 @@ void SProjectLauncherDeployTargets::Construct(const FArguments& InArgs, const FP
 		SNew(SVerticalBox)
 
 		+ SVerticalBox::Slot()
-		.AutoHeight()
-		.HAlign(HAlign_Left)
-		.Padding(10.0f, 10.0f)
-		[			
-			SAssignNew(PlatformComboBox, SComboBox<TSharedPtr<FName>>)			
-			.ContentPadding(FMargin(6.0f, 2.0f))
-			.OptionsSource(&VanillaPlatformList)
-			.OnGenerateWidget(this, &SProjectLauncherDeployTargets::HandlePlatformComboBoxGenerateWidget)
-			.OnSelectionChanged(this, &SProjectLauncherDeployTargets::HandlePlatformComboBoxSelectionChanged)
-			[
-				SNew(STextBlock)
-				.Text(this, &SProjectLauncherDeployTargets::HandlePlatformComboBoxContentText)
+			.AutoHeight()
+			.HAlign(HAlign_Left)
+			.Padding(10.0f, 10.0f)
+			[			
+				SAssignNew(PlatformComboBox, SComboBox<TSharedPtr<FName>>)			
+					.ContentPadding(FMargin(6.0f, 2.0f))
+					.OptionsSource(&VanillaPlatformList)
+					.OnGenerateWidget(this, &SProjectLauncherDeployTargets::HandlePlatformComboBoxGenerateWidget)
+					.OnSelectionChanged(this, &SProjectLauncherDeployTargets::HandlePlatformComboBoxSelectionChanged)
+					[
+						SNew(STextBlock)
+						.Text(this, &SProjectLauncherDeployTargets::HandlePlatformComboBoxContentText)
+					]
 			]
-		]
-
-		+ SVerticalBox::Slot()
-		.AutoHeight()
-		[
-			// device list
-			SAssignNew(DeviceProxyListView, SListView<ITargetDeviceProxyPtr>)
-			.ItemHeight(16.0f)
-			.HeaderRow
-			(
-			SNew(SHeaderRow)
-
-			+ SHeaderRow::Column("CheckBox")
-			.DefaultLabel(LOCTEXT("DeviceListCheckboxColumnHeader", " "))
-			.FixedWidth(24.0f)
-
-			+ SHeaderRow::Column("Device")
-			.DefaultLabel(LOCTEXT("DeviceListDeviceColumnHeader", "Device"))
-			.FillWidth(0.35f)
-
-			+ SHeaderRow::Column("Variant")
-			.DefaultLabel(LOCTEXT("DeviceListVariantColumnHeader", "Variant"))
-			.FillWidth(0.2f)
-
-			+ SHeaderRow::Column("Platform")
-			.DefaultLabel(LOCTEXT("DeviceListPlatformColumnHeader", "Platform"))
-			.FillWidth(0.15f)
-
-			+ SHeaderRow::Column("Host")
-			.DefaultLabel(LOCTEXT("DeviceListHostColumnHeader", "Host"))
-			.FillWidth(0.15f)
-
-			+ SHeaderRow::Column("Owner")
-			.DefaultLabel(LOCTEXT("DeviceListOwnerColumnHeader", "Owner"))
-			.FillWidth(0.15f)
-			)
-			.ListItemsSource(&DeviceProxyList)
-			.OnGenerateRow(this, &SProjectLauncherDeployTargets::HandleDeviceProxyListViewGenerateRow)
-			.SelectionMode(ESelectionMode::Single)
-			.Visibility(this, &SProjectLauncherDeployTargets::HandleDeviceProxyListViewVisibility)
-		]
 
 		+ SVerticalBox::Slot()
-		.AutoHeight()
-		.Padding(0.0f, 12.0f, 0.0f, 0.0f)
-		[
-			SNew(SHorizontalBox)
-			.Visibility(this, &SProjectLauncherDeployTargets::HandleNoDevicesBoxVisibility)
-
-			+ SHorizontalBox::Slot()
-			.AutoWidth()
+			.AutoHeight()
 			[
-				SNew(SImage)
-				.Image(FEditorStyle::GetBrush(TEXT("Icons.Warning")))
+				// device list
+				SAssignNew(DeviceProxyListView, SListView<TSharedPtr<ITargetDeviceProxy>>)
+					.ItemHeight(16.0f)
+					.HeaderRow
+					(
+					SNew(SHeaderRow)
+
+					+ SHeaderRow::Column("CheckBox")
+						.DefaultLabel(LOCTEXT("DeviceListCheckboxColumnHeader", " "))
+						.FixedWidth(24.0f)
+
+					+ SHeaderRow::Column("Device")
+						.DefaultLabel(LOCTEXT("DeviceListDeviceColumnHeader", "Device"))
+						.FillWidth(0.35f)
+
+					+ SHeaderRow::Column("Variant")
+						.DefaultLabel(LOCTEXT("DeviceListVariantColumnHeader", "Variant"))
+						.FillWidth(0.2f)
+
+					+ SHeaderRow::Column("Platform")
+						.DefaultLabel(LOCTEXT("DeviceListPlatformColumnHeader", "Platform"))
+						.FillWidth(0.15f)
+
+					+ SHeaderRow::Column("Host")
+						.DefaultLabel(LOCTEXT("DeviceListHostColumnHeader", "Host"))
+						.FillWidth(0.15f)
+
+					+ SHeaderRow::Column("Owner")
+						.DefaultLabel(LOCTEXT("DeviceListOwnerColumnHeader", "Owner"))
+						.FillWidth(0.15f)
+					)
+				.ListItemsSource(&DeviceProxyList)
+				.OnGenerateRow(this, &SProjectLauncherDeployTargets::HandleDeviceProxyListViewGenerateRow)
+				.SelectionMode(ESelectionMode::Single)
+				.Visibility(this, &SProjectLauncherDeployTargets::HandleDeviceProxyListViewVisibility)
 			]
 
-			+ SHorizontalBox::Slot()
-			.AutoWidth()
-			.Padding(4.0f, 0.0f)
-			.VAlign(VAlign_Center)
+		+ SVerticalBox::Slot()
+			.AutoHeight()
+			.Padding(0.0f, 12.0f, 0.0f, 0.0f)
 			[
-				SNew(STextBlock)
-				.Text(this, &SProjectLauncherDeployTargets::HandleNoDevicesTextBlockText)
+				SNew(SHorizontalBox)
+					.Visibility(this, &SProjectLauncherDeployTargets::HandleNoDevicesBoxVisibility)
+
+				+ SHorizontalBox::Slot()
+				.AutoWidth()
+				[
+					SNew(SImage)
+					.Image(FEditorStyle::GetBrush(TEXT("Icons.Warning")))
+				]
+
+				+ SHorizontalBox::Slot()
+					.AutoWidth()
+					.Padding(4.0f, 0.0f)
+					.VAlign(VAlign_Center)
+					[
+						SNew(STextBlock)
+							.Text(this, &SProjectLauncherDeployTargets::HandleNoDevicesTextBlockText)
+					]
 			]
-		]
 	];
 
-	const ITargetDeviceProxyManagerRef& DeviceProxyManager = Model->GetDeviceProxyManager();
+	const TSharedRef<ITargetDeviceProxyManager>& DeviceProxyManager = Model->GetDeviceProxyManager();
 
 	DeviceProxyManager->OnProxyAdded().AddSP(this, &SProjectLauncherDeployTargets::HandleDeviceProxyManagerProxyAdded);
 	DeviceProxyManager->OnProxyRemoved().AddSP(this, &SProjectLauncherDeployTargets::HandleDeviceProxyManagerProxyRemoved);
-
 	DeviceProxyManager->GetProxies(NAME_None, false, DeviceProxyList);
 
 	VanillaPlatformList.Reset();
+
 	TSet<FName> VanillaNames;
 	VanillaNames.Add(NAME_None);
 	VanillaPlatformList.Add(MakeShareable(new FName(NAME_None)));
@@ -144,6 +147,7 @@ void SProjectLauncherDeployTargets::Construct(const FArguments& InArgs, const FP
 			}
 		}
 	}
+
 	PlatformComboBox->SetSelectedItem(VanillaPlatformList[0]);
 }
 END_SLATE_FUNCTION_BUILD_OPTIMIZATION
@@ -155,7 +159,6 @@ END_SLATE_FUNCTION_BUILD_OPTIMIZATION
 void SProjectLauncherDeployTargets::RefreshDeviceProxyList()
 {
 	Model->GetDeviceProxyManager()->GetProxies(NAME_None, false, DeviceProxyList);
-
 	DeviceProxyListView->RequestListRefresh();
 }
 
@@ -170,11 +173,12 @@ ILauncherDeviceGroupPtr SProjectLauncherDeployTargets::HandleDeviceListRowDevice
 	{
 		return SelectedProfile->GetDeployedDeviceGroup();
 	}
+
 	return nullptr;
 }
 
 
-bool SProjectLauncherDeployTargets::HandleDeviceListRowIsEnabled(ITargetDeviceProxyPtr DeviceProxy) const
+bool SProjectLauncherDeployTargets::HandleDeviceListRowIsEnabled(TSharedPtr<ITargetDeviceProxy> DeviceProxy) const
 {
 	if (DeviceProxy.IsValid())
 	{
@@ -192,7 +196,7 @@ bool SProjectLauncherDeployTargets::HandleDeviceListRowIsEnabled(ITargetDevicePr
 }
 
 
-FText SProjectLauncherDeployTargets::HandleDeviceListRowToolTipText(ITargetDeviceProxyPtr DeviceProxy) const
+FText SProjectLauncherDeployTargets::HandleDeviceListRowToolTipText(TSharedPtr<ITargetDeviceProxy> DeviceProxy) const
 {
 	FTextBuilder Builder;
 	Builder.AppendLineFormat(LOCTEXT("DeviceListRowToolTipName", "Name: {0}"), FText::FromString(DeviceProxy->GetName()));
@@ -210,7 +214,8 @@ FText SProjectLauncherDeployTargets::HandleDeviceListRowToolTipText(ITargetDevic
 	return Builder.ToText();
 }
 
-TSharedRef<ITableRow> SProjectLauncherDeployTargets::HandleDeviceProxyListViewGenerateRow(ITargetDeviceProxyPtr InItem, const TSharedRef<STableViewBase>& OwnerTable) const
+
+TSharedRef<ITableRow> SProjectLauncherDeployTargets::HandleDeviceProxyListViewGenerateRow(TSharedPtr<ITargetDeviceProxy> InItem, const TSharedRef<STableViewBase>& OwnerTable) const
 {
 	check(Model->GetSelectedProfile().IsValid());
 
@@ -221,6 +226,7 @@ TSharedRef<ITableRow> SProjectLauncherDeployTargets::HandleDeviceProxyListViewGe
 		.ToolTipText(this, &SProjectLauncherDeployTargets::HandleDeviceListRowToolTipText, InItem);
 }
 
+
 EVisibility SProjectLauncherDeployTargets::HandleDeviceProxyListViewVisibility() const
 {
 	TSharedPtr<FName> DefaultPlatformNamePtr = PlatformComboBox->GetSelectedItem();
@@ -234,6 +240,7 @@ EVisibility SProjectLauncherDeployTargets::HandleDeviceProxyListViewVisibility()
 	return EVisibility::Collapsed;
 }
 
+
 EVisibility SProjectLauncherDeployTargets::HandleNoDevicesBoxVisibility() const
 {
 	if (DeviceProxyList.Num() == 0)
@@ -244,6 +251,7 @@ EVisibility SProjectLauncherDeployTargets::HandleNoDevicesBoxVisibility() const
 	return EVisibility::Collapsed;
 }
 
+
 FText SProjectLauncherDeployTargets::HandleNoDevicesTextBlockText() const
 {
 	if (DeviceProxyList.Num() == 0)
@@ -254,16 +262,19 @@ FText SProjectLauncherDeployTargets::HandleNoDevicesTextBlockText() const
 	return FText::GetEmpty();
 }
 
-void SProjectLauncherDeployTargets::HandleDeviceProxyManagerProxyAdded(const ITargetDeviceProxyRef& AddedProxy)
+
+void SProjectLauncherDeployTargets::HandleDeviceProxyManagerProxyAdded(const TSharedRef<ITargetDeviceProxy>& AddedProxy)
 {
 	RefreshDeviceProxyList();
 }
 
-void SProjectLauncherDeployTargets::HandleDeviceProxyManagerProxyRemoved(const ITargetDeviceProxyRef& RemovedProxy)
+
+void SProjectLauncherDeployTargets::HandleDeviceProxyManagerProxyRemoved(const TSharedRef<ITargetDeviceProxy>& RemovedProxy)
 {
 	RefreshDeviceProxyList();
 }
 
+
 FText SProjectLauncherDeployTargets::HandlePlatformComboBoxContentText() const
 {	
 	FName DefaultPlatformName = NAME_None;
@@ -296,6 +307,7 @@ FText SProjectLauncherDeployTargets::HandlePlatformComboBoxContentText() const
 	return DefaultPlatformText;
 }
 
+
 TSharedRef<SWidget> SProjectLauncherDeployTargets::HandlePlatformComboBoxGenerateWidget(TSharedPtr<FName> StringItem)
 {
 	FName DefaultPlatformName = StringItem.IsValid() ? *StringItem : NAME_None;	
@@ -303,6 +315,7 @@ TSharedRef<SWidget> SProjectLauncherDeployTargets::HandlePlatformComboBoxGenerat
 	return SNew(STextBlock).Text(FText::FromName(DefaultPlatformName));
 }
 
+
 void SProjectLauncherDeployTargets::HandlePlatformComboBoxSelectionChanged(TSharedPtr<FName> StringItem, ESelectInfo::Type SelectInfo)
 {
 	if (Model.IsValid())
@@ -315,4 +328,5 @@ void SProjectLauncherDeployTargets::HandlePlatformComboBoxSelectionChanged(TShar
 	} 	
 }
 
+
 #undef LOCTEXT_NAMESPACE
diff --git a/Engine/Source/Developer/ProjectLauncher/Private/Widgets/Deploy/SProjectLauncherDeployTargets.h b/Engine/Source/Developer/ProjectLauncher/Private/Widgets/Deploy/SProjectLauncherDeployTargets.h
--- a/Engine/Source/Developer/ProjectLauncher/Private/Widgets/Deploy/SProjectLauncherDeployTargets.h
+++ b/Engine/Source/Developer/ProjectLauncher/Private/Widgets/Deploy/SProjectLauncherDeployTargets.h
@@ -12,6 +12,9 @@
 #include "Widgets/Views/STableRow.h"
 #include "Widgets/Input/SComboBox.h"
 
+class ITargetDeviceProxy;
+
+
 /**
  * Implements the deployment targets panel.
  */
@@ -38,7 +41,7 @@ public:
 	 * @param InArgs The Slate argument list.
 	 * @param InModel The data model.
 	 */
-	void Construct(	const FArguments& InArgs, const FProjectLauncherModelRef& InModel);
+	void Construct(	const FArguments& InArgs, const TSharedRef<FProjectLauncherModel>& InModel);
 
 protected:
 
@@ -49,67 +52,67 @@ protected:
 
 private:
 
-	// Callback for selecting a different device group in the device group selector.
-	void HandleDeviceGroupSelectorGroupSelected( const ILauncherDeviceGroupPtr& SelectedGroup );
+	/** Callback for selecting a different device group in the device group selector. */
+	void HandleDeviceGroupSelectorGroupSelected(const ILauncherDeviceGroupPtr& SelectedGroup);
 
-	// Callback for getting the currently selected device group for a device list row.
-	ILauncherDeviceGroupPtr HandleDeviceListRowDeviceGroup( ) const;
+	/** Callback for getting the currently selected device group for a device list row. */
+	ILauncherDeviceGroupPtr HandleDeviceListRowDeviceGroup() const;
 
-	// Callback for getting the enabled state of a device proxy list row.
-	bool HandleDeviceListRowIsEnabled( ITargetDeviceProxyPtr DeviceProxy ) const;
+	/** Callback for getting the enabled state of a device proxy list row. */
+	bool HandleDeviceListRowIsEnabled(TSharedPtr<ITargetDeviceProxy> DeviceProxy) const;
 
-	// Callback for getting the tool tip text of a device proxy list row.
-	FText HandleDeviceListRowToolTipText( ITargetDeviceProxyPtr DeviceProxy ) const;
+	/** Callback for getting the tool tip text of a device proxy list row. */
+	FText HandleDeviceListRowToolTipText(TSharedPtr<ITargetDeviceProxy> DeviceProxy) const;
 
-	// Callback for generating a row in the device list view.
-	TSharedRef<ITableRow> HandleDeviceProxyListViewGenerateRow( ITargetDeviceProxyPtr InItem, const TSharedRef<STableViewBase>& OwnerTable ) const;
+	/** Callback for generating a row in the device list view. */
+	TSharedRef<ITableRow> HandleDeviceProxyListViewGenerateRow(TSharedPtr<ITargetDeviceProxy> InItem, const TSharedRef<STableViewBase>& OwnerTable) const;
 
-	// Callback for determining the visibility of the device list view.
-	EVisibility HandleDeviceProxyListViewVisibility( ) const;
+	/** Callback for determining the visibility of the device list view. */
+	EVisibility HandleDeviceProxyListViewVisibility() const;
 
-	// Callback for getting a flag indicating whether a valid device group is currently selected.
-	EVisibility HandleDeviceSelectorVisibility( ) const;
+	/** Callback for getting a flag indicating whether a valid device group is currently selected. */
+	EVisibility HandleDeviceSelectorVisibility() const;
 
-	// Callback for getting the visibility of the 'No devices detected' text block.
-	EVisibility HandleNoDevicesBoxVisibility( ) const;
+	/** Callback for getting the visibility of the 'No devices detected' text block. */
+	EVisibility HandleNoDevicesBoxVisibility() const;
 
-	// Callback for getting the text in the 'no devices' text block.
-	FText HandleNoDevicesTextBlockText( ) const;
+	/** Callback for getting the text in the 'no devices' text block. */
+	FText HandleNoDevicesTextBlockText() const;
 
-	// Callback for determining the visibility of the 'Select a device group' text block.
-	EVisibility HandleSelectDeviceGroupWarningVisibility( ) const;
+	/** Callback for determining the visibility of the 'Select a device group' text block. */
+	EVisibility HandleSelectDeviceGroupWarningVisibility() const;
 
-	// Callback for when a device proxy has been added to the device proxy manager.
-	void HandleDeviceProxyManagerProxyAdded(const ITargetDeviceProxyRef& AddedProxy);
+	/** Callback for when a device proxy has been added to the device proxy manager. */
+	void HandleDeviceProxyManagerProxyAdded(const TSharedRef<ITargetDeviceProxy>& AddedProxy);
 
-	// Callback for when a device proxy has been removed from the device proxy manager.
-	void HandleDeviceProxyManagerProxyRemoved(const ITargetDeviceProxyRef& RemovedProxy);
+	/** Callback for when a device proxy has been removed from the device proxy manager. */
+	void HandleDeviceProxyManagerProxyRemoved(const TSharedRef<ITargetDeviceProxy>& RemovedProxy);
 
-	// Callback for getting the name of the selected default platform.
+	/** Callback for getting the name of the selected default platform. */
 	FText HandlePlatformComboBoxContentText() const;
 
-	// Callback for generating widgets for the platforms combo box.
+	/** Callback for generating widgets for the platforms combo box. */
 	TSharedRef<SWidget> HandlePlatformComboBoxGenerateWidget(TSharedPtr<FName> StringItem);
 
-	// Callback for handling default platform selection changes.
+	/** Callback for handling default platform selection changes. */
 	void HandlePlatformComboBoxSelectionChanged(TSharedPtr<FName> StringItem, ESelectInfo::Type SelectInfo);
 
 private:
 
-	// Holds the list of available device proxies.
-	TArray<ITargetDeviceProxyPtr> DeviceProxyList;
+	/** Holds the list of available device proxies. */
+	TArray<TSharedPtr<ITargetDeviceProxy>> DeviceProxyList;
 
-	// Holds the list of vanilla platforms;
+	/** Holds the list of vanilla platforms. */
 	TArray<TSharedPtr<FName>> VanillaPlatformList;
 
-	// Holds the list of vanilla platform options.
+	/** Holds the list of vanilla platform options. */
 	TSharedPtr<SComboBox<TSharedPtr<FName>>> PlatformComboBox;
 	mutable FText DefaultPlatformText;
 
-	// Holds the device proxy list view .
-	TSharedPtr<SListView<ITargetDeviceProxyPtr> > DeviceProxyListView;
+	/** Holds the device proxy list view . */
+	TSharedPtr<SListView<TSharedPtr<ITargetDeviceProxy>> > DeviceProxyListView;
 
-	// Holds a pointer to the data model.
-	FProjectLauncherModelPtr Model;
+	/** Holds a pointer to the data model. */
+	TSharedPtr<FProjectLauncherModel> Model;
 
 };
diff --git a/Engine/Source/Developer/ProjectLauncher/Private/Widgets/Deploy/SProjectLauncherDeployToDeviceSettings.cpp b/Engine/Source/Developer/ProjectLauncher/Private/Widgets/Deploy/SProjectLauncherDeployToDeviceSettings.cpp
--- a/Engine/Source/Developer/ProjectLauncher/Private/Widgets/Deploy/SProjectLauncherDeployToDeviceSettings.cpp
+++ b/Engine/Source/Developer/ProjectLauncher/Private/Widgets/Deploy/SProjectLauncherDeployToDeviceSettings.cpp
@@ -1,14 +1,16 @@
 // Copyright 1998-2017 Epic Games, Inc. All Rights Reserved.
 
-#include "Widgets/Deploy/SProjectLauncherDeployToDeviceSettings.h"
-#include "Widgets/SBoxPanel.h"
+#include "SProjectLauncherDeployToDeviceSettings.h"
+
+#include "EditorStyleSet.h"
 #include "Styling/SlateTypes.h"
+#include "Widgets/SBoxPanel.h"
+#include "Widgets/Input/SCheckBox.h"
 #include "Widgets/Layout/SBorder.h"
+#include "Widgets/Layout/SExpandableArea.h"
 #include "Widgets/Text/STextBlock.h"
-#include "Widgets/Input/SCheckBox.h"
-#include "EditorStyleSet.h"
+
 #include "Widgets/Deploy/SProjectLauncherDeployTargets.h"
-#include "Widgets/Layout/SExpandableArea.h"
 
 
 #define LOCTEXT_NAMESPACE "SProjectLauncherDeployToDeviceSettings"
@@ -17,7 +19,7 @@
 /* SProjectLauncherDeployToDeviceSettings interface
  *****************************************************************************/
 
-void SProjectLauncherDeployToDeviceSettings::Construct( const FArguments& InArgs, const FProjectLauncherModelRef& InModel, EVisibility InShowAdvanced )
+void SProjectLauncherDeployToDeviceSettings::Construct(const FArguments& InArgs, const TSharedRef<FProjectLauncherModel>& InModel, EVisibility InShowAdvanced)
 {
 	Model = InModel;
 
@@ -80,8 +82,7 @@ void SProjectLauncherDeployToDeviceSettings::Construct( const FArguments& InArgs
 /* SProjectLauncherDeployToDeviceSettings callbacks
  *****************************************************************************/
 
-
-void SProjectLauncherDeployToDeviceSettings::HandleIncrementalCheckBoxCheckStateChanged( ECheckBoxState NewState )
+void SProjectLauncherDeployToDeviceSettings::HandleIncrementalCheckBoxCheckStateChanged(ECheckBoxState NewState)
 {
 	ILauncherProfilePtr SelectedProfile = Model->GetSelectedProfile();
 
@@ -92,7 +93,7 @@ void SProjectLauncherDeployToDeviceSettings::HandleIncrementalCheckBoxCheckState
 }
 
 
-ECheckBoxState SProjectLauncherDeployToDeviceSettings::HandleIncrementalCheckBoxIsChecked( ) const
+ECheckBoxState SProjectLauncherDeployToDeviceSettings::HandleIncrementalCheckBoxIsChecked() const
 {
 	ILauncherProfilePtr SelectedProfile = Model->GetSelectedProfile();
 
@@ -106,5 +107,4 @@ ECheckBoxState SProjectLauncherDeployToDeviceSettings::HandleIncrementalCheckBox
 
 	return ECheckBoxState::Unchecked;
 }
-
 #undef LOCTEXT_NAMESPACE
diff --git a/Engine/Source/Developer/ProjectLauncher/Private/Widgets/Deploy/SProjectLauncherDeployToDeviceSettings.h b/Engine/Source/Developer/ProjectLauncher/Private/Widgets/Deploy/SProjectLauncherDeployToDeviceSettings.h
--- a/Engine/Source/Developer/ProjectLauncher/Private/Widgets/Deploy/SProjectLauncherDeployToDeviceSettings.h
+++ b/Engine/Source/Developer/ProjectLauncher/Private/Widgets/Deploy/SProjectLauncherDeployToDeviceSettings.h
@@ -29,7 +29,7 @@ public:
 	 * @param InArgs - The Slate argument list.
 	 * @param InModel - The data model.
 	 */
-	void Construct(	const FArguments& InArgs, const FProjectLauncherModelRef& InModel, EVisibility InShowAdvanced = EVisibility::Visible );
+	void Construct(	const FArguments& InArgs, const TSharedRef<FProjectLauncherModel>& InModel, EVisibility InShowAdvanced = EVisibility::Visible );
 
 protected:
 
@@ -42,5 +42,5 @@ private:
 	ECheckBoxState HandleIncrementalCheckBoxIsChecked( ) const;
 
 	// Holds a pointer to the data model.
-	FProjectLauncherModelPtr Model;
+	TSharedPtr<FProjectLauncherModel> Model;
 };
diff --git a/Engine/Source/Developer/ProjectLauncher/Private/Widgets/Deploy/SProjectLauncherDeviceGroupSelector.h b/Engine/Source/Developer/ProjectLauncher/Private/Widgets/Deploy/SProjectLauncherDeviceGroupSelector.h
--- a/Engine/Source/Developer/ProjectLauncher/Private/Widgets/Deploy/SProjectLauncherDeviceGroupSelector.h
+++ b/Engine/Source/Developer/ProjectLauncher/Private/Widgets/Deploy/SProjectLauncherDeviceGroupSelector.h
@@ -3,8 +3,8 @@
 #pragma once
 
 #include "CoreMinimal.h"
+#include "ILauncherProfileManager.h"
 #include "SlateFwd.h"
-#include "Interfaces/ILauncherProfileManager.h"
 #include "Input/Reply.h"
 #include "Widgets/DeclarativeSyntaxSupport.h"
 #include "Widgets/SWidget.h"
diff --git a/Engine/Source/Developer/ProjectLauncher/Private/Widgets/Deploy/SProjectLauncherSimpleDeviceListRow.h b/Engine/Source/Developer/ProjectLauncher/Private/Widgets/Deploy/SProjectLauncherSimpleDeviceListRow.h
--- a/Engine/Source/Developer/ProjectLauncher/Private/Widgets/Deploy/SProjectLauncherSimpleDeviceListRow.h
+++ b/Engine/Source/Developer/ProjectLauncher/Private/Widgets/Deploy/SProjectLauncherSimpleDeviceListRow.h
@@ -15,7 +15,7 @@
 #include "Widgets/Layout/SBox.h"
 #include "Widgets/Views/STableRow.h"
 #include "Models/ProjectLauncherModel.h"
-#include "Widgets/SProjectLauncherDelegates.h"
+#include "Widgets/Shared/ProjectLauncherDelegates.h"
 #include "Widgets/Text/STextBlock.h"
 #include "Widgets/Views/SListView.h"
 #include "Widgets/Images/SImage.h"
@@ -33,7 +33,7 @@
  * Implements a row widget for the launcher's device proxy list.
  */
 class SProjectLauncherSimpleDeviceListRow
-	: public STableRow<ITargetDeviceProxyPtr>
+	: public STableRow<TSharedPtr<ITargetDeviceProxy>>
 {
 public:
 
@@ -46,7 +46,7 @@ public:
 		/**
 		 * The device proxy shown in this row.
 		 */
-		SLATE_ARGUMENT(ITargetDeviceProxyPtr, DeviceProxy)
+		SLATE_ARGUMENT(TSharedPtr<ITargetDeviceProxy>, DeviceProxy)
 
 		/**
 		 * Whether the advanced options should be shown.
@@ -71,9 +71,9 @@ public:
 	 * @param InArgs - The construction arguments.
 	 * @param InModel - The launcher model this list uses.
 	 */
-	void Construct(const FArguments& InArgs, const FProjectLauncherModelRef& InModel, const TSharedRef<STableViewBase>& InOwnerTableView)
+	void Construct(const FArguments& InArgs, const TSharedRef<FProjectLauncherModel>& InModel, const TSharedRef<STableViewBase>& InOwnerTableView)
 	{
-		STableRow< ITargetDeviceProxyPtr >::ConstructInternal(
+		STableRow< TSharedPtr<ITargetDeviceProxy> >::ConstructInternal(
 			STableRow::FArguments()
 			.ShowSelection(false),
 			InOwnerTableView
@@ -468,10 +468,10 @@ private:
 private:
 
 	// Holds a pointer to the data model.
-	FProjectLauncherModelPtr Model;
+	TSharedPtr<FProjectLauncherModel> Model;
 
 	// Holds a reference to the device proxy that is displayed in this row.
-	ITargetDeviceProxyPtr DeviceProxy;
+	TSharedPtr<ITargetDeviceProxy> DeviceProxy;
 
 	// Holds a reference to the simple profile for this device
 	ILauncherSimpleProfilePtr SimpleProfile;
diff --git a/Engine/Source/Developer/ProjectLauncher/Private/Widgets/Deploy/SProjectLauncherSimpleDeviceListView.cpp b/Engine/Source/Developer/ProjectLauncher/Private/Widgets/Deploy/SProjectLauncherSimpleDeviceListView.cpp
--- a/Engine/Source/Developer/ProjectLauncher/Private/Widgets/Deploy/SProjectLauncherSimpleDeviceListView.cpp
+++ b/Engine/Source/Developer/ProjectLauncher/Private/Widgets/Deploy/SProjectLauncherSimpleDeviceListView.cpp
@@ -1,13 +1,19 @@
 // Copyright 1998-2017 Epic Games, Inc. All Rights Reserved.
 
-#include "Widgets/Deploy/SProjectLauncherSimpleDeviceListView.h"
+#include "SProjectLauncherSimpleDeviceListView.h"
+
+#include "Framework/Docking/TabManager.h"
+#include "ITargetDeviceProxy.h"
+#include "ITargetDeviceProxyManager.h"
 #include "SlateOptMacros.h"
+#include "Widgets/Input/SHyperlink.h"
 #include "Widgets/Text/STextBlock.h"
 #include "Widgets/Layout/SScrollBorder.h"
-#include "Framework/Docking/TabManager.h"
-#include "Widgets/Views/SListView.h"
+#include "Widgets/Views/STableRow.h"
+#include "Widgets/Views/STableViewBase.h"
+
+#include "Models/ProjectLauncherModel.h"
 #include "Widgets/Deploy/SProjectLauncherSimpleDeviceListRow.h"
-#include "Widgets/Input/SHyperlink.h"
 
 
 #define LOCTEXT_NAMESPACE "SProjectLauncherSimpleDeviceListView"
@@ -16,11 +22,11 @@
 /* SProjectLauncherDeployTargets structors
  *****************************************************************************/
 
-SProjectLauncherSimpleDeviceListView::~SProjectLauncherSimpleDeviceListView( )
+SProjectLauncherSimpleDeviceListView::~SProjectLauncherSimpleDeviceListView()
 {
 	if (Model.IsValid())
 	{
-		const ITargetDeviceProxyManagerRef& DeviceProxyManager = Model->GetDeviceProxyManager();
+		const TSharedRef<ITargetDeviceProxyManager>& DeviceProxyManager = Model->GetDeviceProxyManager();
 		DeviceProxyManager->OnProxyAdded().RemoveAll(this);
 		DeviceProxyManager->OnProxyRemoved().RemoveAll(this);
 	}
@@ -31,14 +37,14 @@ SProjectLauncherSimpleDeviceListView::~SProjectLauncherSimpleDeviceListView( )
  *****************************************************************************/
 
 BEGIN_SLATE_FUNCTION_BUILD_OPTIMIZATION
-void SProjectLauncherSimpleDeviceListView::Construct(const FArguments& InArgs, const FProjectLauncherModelRef& InModel)
+void SProjectLauncherSimpleDeviceListView::Construct(const FArguments& InArgs, const TSharedRef<FProjectLauncherModel>& InModel)
 {
 	OnProfileRun = InArgs._OnProfileRun;
 	IsAdvanced = InArgs._IsAdvanced;
 	
 	Model = InModel;
 
-	SAssignNew(DeviceProxyListView, SListView<ITargetDeviceProxyPtr>)
+	SAssignNew(DeviceProxyListView, SListView<TSharedPtr<ITargetDeviceProxy>>)
 		.SelectionMode(ESelectionMode::None)
 		.ListItemsSource(&DeviceProxyList)
 		.OnGenerateRow(this, &SProjectLauncherSimpleDeviceListView::HandleDeviceProxyListViewGenerateRow)
@@ -87,11 +93,10 @@ void SProjectLauncherSimpleDeviceListView::Construct(const FArguments& InArgs, c
 		]
 	];
 
-	const ITargetDeviceProxyManagerRef& DeviceProxyManager = Model->GetDeviceProxyManager();
+	const TSharedRef<ITargetDeviceProxyManager>& DeviceProxyManager = Model->GetDeviceProxyManager();
 
 	DeviceProxyManager->OnProxyAdded().AddSP(this, &SProjectLauncherSimpleDeviceListView::HandleDeviceProxyManagerProxyAdded);
 	DeviceProxyManager->OnProxyRemoved().AddSP(this, &SProjectLauncherSimpleDeviceListView::HandleDeviceProxyManagerProxyRemoved);
-
 	DeviceProxyManager->GetProxies(NAME_None, false, DeviceProxyList);
 
 }
@@ -111,28 +116,30 @@ void SProjectLauncherSimpleDeviceListView::RefreshDeviceProxyList()
 /* SProjectLauncherDeployTargets callbacks
  *****************************************************************************/
 
-bool SProjectLauncherSimpleDeviceListView::HandleDeviceListRowIsEnabled( ITargetDeviceProxyPtr DeviceProxy ) const
+bool SProjectLauncherSimpleDeviceListView::HandleDeviceListRowIsEnabled(TSharedPtr<ITargetDeviceProxy> DeviceProxy) const
 {
 	return true;
 }
 
+
 void SProjectLauncherSimpleDeviceListView::HandleDeviceManagerHyperlinkNavigate() const
 {
 	FGlobalTabmanager::Get()->InvokeTab(FTabId("DeviceManager"));
 }
 
-FText SProjectLauncherSimpleDeviceListView::HandleDeviceListRowToolTipText( ITargetDeviceProxyPtr DeviceProxy ) const
+
+FText SProjectLauncherSimpleDeviceListView::HandleDeviceListRowToolTipText(TSharedPtr<ITargetDeviceProxy> DeviceProxy) const
 {
 	FTextBuilder Builder;
 	Builder.AppendLineFormat(LOCTEXT("DeviceListRowToolTipName", "Name: {0}"), FText::FromString(DeviceProxy->GetName()));
-	//Builder.AppendLineFormat(LOCTEXT("DeviceListRowToolTipPlatform", "Platform: {0}"), FText::FromString(DeviceProxy->GetTargetPlatformName(SimpleProfile->GetDeviceVariant())));
-	//Builder.AppendLineFormat(LOCTEXT("DeviceListRowToolTipDeviceId", "Device ID: {0}"), FText::FromString(DeviceProxy->GetTargetDeviceId(SimpleProfile->GetDeviceVariant())));
+//	Builder.AppendLineFormat(LOCTEXT("DeviceListRowToolTipPlatform", "Platform: {0}"), FText::FromString(DeviceProxy->GetTargetPlatformName(SimpleProfile->GetDeviceVariant())));
+//	Builder.AppendLineFormat(LOCTEXT("DeviceListRowToolTipDeviceId", "Device ID: {0}"), FText::FromString(DeviceProxy->GetTargetDeviceId(SimpleProfile->GetDeviceVariant())));
 
 	return Builder.ToText();
 }
 
 
-TSharedRef<ITableRow> SProjectLauncherSimpleDeviceListView::HandleDeviceProxyListViewGenerateRow(ITargetDeviceProxyPtr InItem, const TSharedRef<STableViewBase>& OwnerTable) const
+TSharedRef<ITableRow> SProjectLauncherSimpleDeviceListView::HandleDeviceProxyListViewGenerateRow(TSharedPtr<ITargetDeviceProxy> InItem, const TSharedRef<STableViewBase>& OwnerTable) const
 {
 	return SNew(SProjectLauncherSimpleDeviceListRow, Model.ToSharedRef(), OwnerTable)
 		.OnProfileRun(OnProfileRun)
@@ -142,14 +149,17 @@ TSharedRef<ITableRow> SProjectLauncherSimpleDeviceListView::HandleDeviceProxyLis
 		.ToolTipText(this, &SProjectLauncherSimpleDeviceListView::HandleDeviceListRowToolTipText, InItem);
 }
 
-void SProjectLauncherSimpleDeviceListView::HandleDeviceProxyManagerProxyAdded(const ITargetDeviceProxyRef& AddedProxy)
+
+void SProjectLauncherSimpleDeviceListView::HandleDeviceProxyManagerProxyAdded(const TSharedRef<ITargetDeviceProxy>& AddedProxy)
 {
 	RefreshDeviceProxyList();
 }
 
-void SProjectLauncherSimpleDeviceListView::HandleDeviceProxyManagerProxyRemoved(const ITargetDeviceProxyRef& AddedProxy)
+
+void SProjectLauncherSimpleDeviceListView::HandleDeviceProxyManagerProxyRemoved(const TSharedRef<ITargetDeviceProxy>& AddedProxy)
 {
 	RefreshDeviceProxyList();
 }
 
+
 #undef LOCTEXT_NAMESPACE
diff --git a/Engine/Source/Developer/ProjectLauncher/Private/Widgets/Deploy/SProjectLauncherSimpleDeviceListView.h b/Engine/Source/Developer/ProjectLauncher/Private/Widgets/Deploy/SProjectLauncherSimpleDeviceListView.h
--- a/Engine/Source/Developer/ProjectLauncher/Private/Widgets/Deploy/SProjectLauncherSimpleDeviceListView.h
+++ b/Engine/Source/Developer/ProjectLauncher/Private/Widgets/Deploy/SProjectLauncherSimpleDeviceListView.h
@@ -2,15 +2,21 @@
 
 #pragma once
 
-#include "CoreMinimal.h"
-#include "Misc/Attribute.h"
+#include "Containers/Array.h"
 #include "Layout/Visibility.h"
+#include "Misc/Attribute.h"
+#include "Templates/SharedPointer.h"
 #include "Widgets/DeclarativeSyntaxSupport.h"
 #include "Widgets/SCompoundWidget.h"
-#include "Widgets/Views/STableViewBase.h"
-#include "Widgets/Views/STableRow.h"
-#include "Models/ProjectLauncherModel.h"
-#include "Widgets/SProjectLauncherDelegates.h"
+#include "Widgets/Views/SListView.h"
+
+#include "Widgets/Shared/ProjectLauncherDelegates.h"
+
+class FProjectLauncherModel;
+class ITableRow;
+class ITargetDeviceProxy;
+class STableViewBase;
+
 
 /**
  * Implements the deployment targets panel.
@@ -27,10 +33,8 @@ public:
 
 public:
 
-	/**
-	 * Destructor.
-	 */
-	~SProjectLauncherSimpleDeviceListView( );
+	/** Destructor. */
+	~SProjectLauncherSimpleDeviceListView();
 
 public:
 
@@ -41,52 +45,47 @@ public:
 	 * @param InModel The data model.
 	 * @param InAdvanced Whether or not the elements should show advanced data.
 	 */
-	void Construct(const FArguments& InArgs, const FProjectLauncherModelRef& InModel);
+	void Construct(const FArguments& InArgs, const TSharedRef<FProjectLauncherModel>& InModel);
 
 protected:
 
-	/**
-	 * Refreshes the list of device proxies.
-	 */
+	/** Refresh the list of device proxies. */
 	void RefreshDeviceProxyList();
 
 private:
 
-	// Callback for getting the enabled state of a device proxy list row.
-	bool HandleDeviceListRowIsEnabled(ITargetDeviceProxyPtr DeviceProxy) const;
+	/** Callback for getting the enabled state of a device proxy list row. */
+	bool HandleDeviceListRowIsEnabled(TSharedPtr<ITargetDeviceProxy> DeviceProxy) const;
 
-	// Callback when the user clicks the Device Manager hyperlink.
+	/** Callback when the user clicks the Device Manager hyperlink. */
 	void HandleDeviceManagerHyperlinkNavigate() const;
 
-	// Callback for getting the tool tip text of a device proxy list row.
-	FText HandleDeviceListRowToolTipText(ITargetDeviceProxyPtr DeviceProxy) const;
-
-	// Callback for generating a row in the device list view.
-	TSharedRef<ITableRow> HandleDeviceProxyListViewGenerateRow(ITargetDeviceProxyPtr InItem, const TSharedRef<STableViewBase>& OwnerTable) const;
+	/** Callback for getting the tool tip text of a device proxy list row. */
+	FText HandleDeviceListRowToolTipText(TSharedPtr<ITargetDeviceProxy> DeviceProxy) const;
 
-	// Callback for determining the visibility of the device list view.
-	EVisibility HandleDeviceProxyListViewVisibility() const;
+	/** Callback for generating a row in the device list view. */
+	TSharedRef<ITableRow> HandleDeviceProxyListViewGenerateRow(TSharedPtr<ITargetDeviceProxy> InItem, const TSharedRef<STableViewBase>& OwnerTable) const;
 
-	// Callback for when a device proxy has been added to the device proxy manager.
-	void HandleDeviceProxyManagerProxyAdded(const ITargetDeviceProxyRef& AddedProxy);
+	/** Callback for when a device proxy has been added to the device proxy manager. */
+	void HandleDeviceProxyManagerProxyAdded(const TSharedRef<ITargetDeviceProxy>& AddedProxy);
 
-	// Callback for when a device proxy has been added to the device proxy manager.
-	void HandleDeviceProxyManagerProxyRemoved(const ITargetDeviceProxyRef& RemovedProxy);
+	/** Callback for when a device proxy has been added to the device proxy manager. */
+	void HandleDeviceProxyManagerProxyRemoved(const TSharedRef<ITargetDeviceProxy>& RemovedProxy);
 
 private:
 
-	// Holds the list of available device proxies.
-	TArray<ITargetDeviceProxyPtr> DeviceProxyList;
+	/** Holds the list of available device proxies. */
+	TArray<TSharedPtr<ITargetDeviceProxy>> DeviceProxyList;
 
-	// Holds the device proxy list view .
-	TSharedPtr<SListView<ITargetDeviceProxyPtr> > DeviceProxyListView;
+	/** Holds the device proxy list view . */
+	TSharedPtr<SListView<TSharedPtr<ITargetDeviceProxy>> > DeviceProxyListView;
 
-	// Holds a pointer to the data model.
-	FProjectLauncherModelPtr Model;
+	/** Holds a pointer to the data model. */
+	TSharedPtr<FProjectLauncherModel> Model;
 
-	// Specifies whether advanced options are shown.
+	/** Specifies whether advanced options are shown. */
 	TAttribute<bool> IsAdvanced;
 
-	// Holds a delegate to be invoked when a profile is run.
+	/** Holds a delegate to be invoked when a profile is run. */
 	FOnProfileRun OnProfileRun;
 };
diff --git a/Engine/Source/Developer/ProjectLauncher/Private/Widgets/Launch/SProjectLauncherLaunchCustomRoles.cpp b/Engine/Source/Developer/ProjectLauncher/Private/Widgets/Launch/SProjectLauncherLaunchCustomRoles.cpp
--- a/Engine/Source/Developer/ProjectLauncher/Private/Widgets/Launch/SProjectLauncherLaunchCustomRoles.cpp
+++ b/Engine/Source/Developer/ProjectLauncher/Private/Widgets/Launch/SProjectLauncherLaunchCustomRoles.cpp
@@ -1,6 +1,6 @@
 // Copyright 1998-2017 Epic Games, Inc. All Rights Reserved.
 
-#include "Widgets/Launch/SProjectLauncherLaunchCustomRoles.h"
+#include "SProjectLauncherLaunchCustomRoles.h"
 
 
 #define LOCTEXT_NAMESPACE "SProjectLauncherLaunchCustomRoles"
@@ -9,7 +9,7 @@
 /* SProjectLauncherLaunchCustomRoles structors
  *****************************************************************************/
 
-SProjectLauncherLaunchCustomRoles::~SProjectLauncherLaunchCustomRoles( )
+SProjectLauncherLaunchCustomRoles::~SProjectLauncherLaunchCustomRoles()
 {
 	if (Model.IsValid())
 	{
@@ -21,7 +21,7 @@ SProjectLauncherLaunchCustomRoles::~SProjectLauncherLaunchCustomRoles( )
 /* SProjectLauncherLaunchCustomRoles interface
  *****************************************************************************/
 
-void SProjectLauncherLaunchCustomRoles::Construct( const FArguments& InArgs, const FProjectLauncherModelRef& InModel )
+void SProjectLauncherLaunchCustomRoles::Construct(const FArguments& InArgs, const TSharedRef<FProjectLauncherModel>& InModel)
 {
 	Model = InModel;
 
diff --git a/Engine/Source/Developer/ProjectLauncher/Private/Widgets/Launch/SProjectLauncherLaunchCustomRoles.h b/Engine/Source/Developer/ProjectLauncher/Private/Widgets/Launch/SProjectLauncherLaunchCustomRoles.h
--- a/Engine/Source/Developer/ProjectLauncher/Private/Widgets/Launch/SProjectLauncherLaunchCustomRoles.h
+++ b/Engine/Source/Developer/ProjectLauncher/Private/Widgets/Launch/SProjectLauncherLaunchCustomRoles.h
@@ -33,10 +33,10 @@ public:
 	 * @param InArgs - The Slate argument list.
 	 * @param InModel - The data model.
 	 */
-	void Construct( const FArguments& InArgs, const FProjectLauncherModelRef& InModel );
+	void Construct( const FArguments& InArgs, const TSharedRef<FProjectLauncherModel>& InModel );
 
 private:
 
 	// Holds a pointer to the data model.
-	FProjectLauncherModelPtr Model;
+	TSharedPtr<FProjectLauncherModel> Model;
 };
diff --git a/Engine/Source/Developer/ProjectLauncher/Private/Widgets/Launch/SProjectLauncherLaunchPage.cpp b/Engine/Source/Developer/ProjectLauncher/Private/Widgets/Launch/SProjectLauncherLaunchPage.cpp
--- a/Engine/Source/Developer/ProjectLauncher/Private/Widgets/Launch/SProjectLauncherLaunchPage.cpp
+++ b/Engine/Source/Developer/ProjectLauncher/Private/Widgets/Launch/SProjectLauncherLaunchPage.cpp
@@ -1,18 +1,20 @@
 // Copyright 1998-2017 Epic Games, Inc. All Rights Reserved.
 
-#include "Widgets/Launch/SProjectLauncherLaunchPage.h"
-#include "Widgets/SBoxPanel.h"
+#include "SProjectLauncherLaunchPage.h"
+
+#include "EditorStyleSet.h"
+#include "Framework/Commands/UIAction.h"
+#include "Framework/MultiBox/MultiBoxBuilder.h"
 #include "SlateOptMacros.h"
 #include "Textures/SlateIcon.h"
-#include "Framework/Commands/UIAction.h"
-#include "Widgets/Layout/SBorder.h"
+#include "Widgets/SBoxPanel.h"
 #include "Widgets/Images/SImage.h"
-#include "Widgets/Text/STextBlock.h"
-#include "Framework/MultiBox/MultiBoxBuilder.h"
 #include "Widgets/Input/SComboButton.h"
-#include "EditorStyleSet.h"
-#include "Widgets/Launch/SProjectLauncherLaunchRoleEditor.h"
+#include "Widgets/Layout/SBorder.h"
 #include "Widgets/Layout/SExpandableArea.h"
+#include "Widgets/Text/STextBlock.h"
+
+#include "Widgets/Launch/SProjectLauncherLaunchRoleEditor.h"
 
 
 #define LOCTEXT_NAMESPACE "SProjectLauncherLaunchPage"
@@ -21,7 +23,7 @@
 /* SProjectLauncherLaunchPage structors
  *****************************************************************************/
 
-SProjectLauncherLaunchPage::~SProjectLauncherLaunchPage( )
+SProjectLauncherLaunchPage::~SProjectLauncherLaunchPage()
 {
 	if (Model.IsValid())
 	{
@@ -34,7 +36,7 @@ SProjectLauncherLaunchPage::~SProjectLauncherLaunchPage( )
  *****************************************************************************/
 
 BEGIN_SLATE_FUNCTION_BUILD_OPTIMIZATION
-void SProjectLauncherLaunchPage::Construct( const FArguments& InArgs, const FProjectLauncherModelRef& InModel )
+void SProjectLauncherLaunchPage::Construct(const FArguments& InArgs, const TSharedRef<FProjectLauncherModel>& InModel)
 {
 	Model = InModel;
 
@@ -158,7 +160,7 @@ void SProjectLauncherLaunchPage::Construct( const FArguments& InArgs, const FPro
 END_SLATE_FUNCTION_BUILD_OPTIMIZATION
 
 
-void SProjectLauncherLaunchPage::Refresh( )
+void SProjectLauncherLaunchPage::Refresh()
 {
 	ILauncherProfilePtr SelectedProfile = Model->GetSelectedProfile();
 
@@ -190,7 +192,7 @@ void SProjectLauncherLaunchPage::Refresh( )
 /* SProjectLauncherLaunchPage callbacks
  *****************************************************************************/
 
-EVisibility SProjectLauncherLaunchPage::HandleCannotLaunchTextBlockVisibility( ) const
+EVisibility SProjectLauncherLaunchPage::HandleCannotLaunchTextBlockVisibility() const
 {
 	ILauncherProfilePtr SelectedProfile = Model->GetSelectedProfile();
 
@@ -206,7 +208,7 @@ EVisibility SProjectLauncherLaunchPage::HandleCannotLaunchTextBlockVisibility( )
 }
 
 
-EVisibility SProjectLauncherLaunchPage::HandleLaunchModeBoxVisibility( ) const
+EVisibility SProjectLauncherLaunchPage::HandleLaunchModeBoxVisibility() const
 {
 	ILauncherProfilePtr SelectedProfile = Model->GetSelectedProfile();
 
@@ -222,7 +224,7 @@ EVisibility SProjectLauncherLaunchPage::HandleLaunchModeBoxVisibility( ) const
 }
 
 
-FText SProjectLauncherLaunchPage::HandleLaunchModeComboButtonContentText( ) const
+FText SProjectLauncherLaunchPage::HandleLaunchModeComboButtonContentText() const
 {
 	ILauncherProfilePtr SelectedProfile = Model->GetSelectedProfile();
 
@@ -252,7 +254,7 @@ FText SProjectLauncherLaunchPage::HandleLaunchModeComboButtonContentText( ) cons
 }
 
 
-void SProjectLauncherLaunchPage::HandleLaunchModeMenuEntryClicked( ELauncherProfileLaunchModes::Type LaunchMode )
+void SProjectLauncherLaunchPage::HandleLaunchModeMenuEntryClicked(ELauncherProfileLaunchModes::Type LaunchMode)
 {
 	ILauncherProfilePtr SelectedProfile = Model->GetSelectedProfile();
 
@@ -263,7 +265,7 @@ void SProjectLauncherLaunchPage::HandleLaunchModeMenuEntryClicked( ELauncherProf
 }
 
 
-EVisibility SProjectLauncherLaunchPage::HandleLaunchSettingsVisibility( ) const
+EVisibility SProjectLauncherLaunchPage::HandleLaunchSettingsVisibility() const
 {
 	ILauncherProfilePtr SelectedProfile = Model->GetSelectedProfile();
 
@@ -282,7 +284,7 @@ EVisibility SProjectLauncherLaunchPage::HandleLaunchSettingsVisibility( ) const
 }
 
 
-void SProjectLauncherLaunchPage::HandleProfileManagerProfileSelected( const ILauncherProfilePtr& SelectedProfile, const ILauncherProfilePtr& PreviousProfile )
+void SProjectLauncherLaunchPage::HandleProfileManagerProfileSelected(const ILauncherProfilePtr& SelectedProfile, const ILauncherProfilePtr& PreviousProfile)
 {
 	if (PreviousProfile.IsValid())
 	{
@@ -295,13 +297,14 @@ void SProjectLauncherLaunchPage::HandleProfileManagerProfileSelected( const ILau
 	Refresh();
 }
 
+
 void SProjectLauncherLaunchPage::HandleProfileProjectChanged()
 {
 	Refresh();
 }
 
 
-EVisibility SProjectLauncherLaunchPage::HandleValidationErrorIconVisibility( ELauncherProfileValidationErrors::Type Error ) const
+EVisibility SProjectLauncherLaunchPage::HandleValidationErrorIconVisibility(ELauncherProfileValidationErrors::Type Error) const
 {
 	ILauncherProfilePtr SelectedProfile = Model->GetSelectedProfile();
 
diff --git a/Engine/Source/Developer/ProjectLauncher/Private/Widgets/Launch/SProjectLauncherLaunchRoleEditor.h b/Engine/Source/Developer/ProjectLauncher/Private/Widgets/Launch/SProjectLauncherLaunchRoleEditor.h
--- a/Engine/Source/Developer/ProjectLauncher/Private/Widgets/Launch/SProjectLauncherLaunchRoleEditor.h
+++ b/Engine/Source/Developer/ProjectLauncher/Private/Widgets/Launch/SProjectLauncherLaunchRoleEditor.h
@@ -3,12 +3,12 @@
 #pragma once
 
 #include "CoreMinimal.h"
+#include "ILauncherProfileLaunchRole.h"
 #include "SlateFwd.h"
 #include "Layout/Visibility.h"
 #include "Widgets/DeclarativeSyntaxSupport.h"
 #include "Styling/SlateColor.h"
 #include "Widgets/SCompoundWidget.h"
-#include "Interfaces/ILauncherProfileLaunchRole.h"
 #include "Styling/SlateTypes.h"
 #include "Widgets/Input/SComboBox.h"
 
diff --git a/Engine/Source/Developer/ProjectLauncher/Private/Widgets/Package/SProjectLauncherPackagePage.cpp b/Engine/Source/Developer/ProjectLauncher/Private/Widgets/Package/SProjectLauncherPackagePage.cpp
--- a/Engine/Source/Developer/ProjectLauncher/Private/Widgets/Package/SProjectLauncherPackagePage.cpp
+++ b/Engine/Source/Developer/ProjectLauncher/Private/Widgets/Package/SProjectLauncherPackagePage.cpp
@@ -1,12 +1,14 @@
 // Copyright 1998-2017 Epic Games, Inc. All Rights Reserved.
 
-#include "Widgets/Package/SProjectLauncherPackagePage.h"
-#include "Widgets/SBoxPanel.h"
-#include "Textures/SlateIcon.h"
+#include "SProjectLauncherPackagePage.h"
+
 #include "Framework/Commands/UIAction.h"
-#include "Widgets/Text/STextBlock.h"
 #include "Framework/MultiBox/MultiBoxBuilder.h"
+#include "Textures/SlateIcon.h"
+#include "Widgets/SBoxPanel.h"
 #include "Widgets/Input/SComboButton.h"
+#include "Widgets/Text/STextBlock.h"
+
 #include "Widgets/Package/SProjectLauncherPackagingSettings.h"
 
 
@@ -16,7 +18,7 @@
 /* SProjectLauncherPackagePage structors
  *****************************************************************************/
 
-SProjectLauncherPackagePage::~SProjectLauncherPackagePage( )
+SProjectLauncherPackagePage::~SProjectLauncherPackagePage()
 {
 	if (Model.IsValid())
 	{
@@ -28,7 +30,7 @@ SProjectLauncherPackagePage::~SProjectLauncherPackagePage( )
 /* SProjectLauncherPackagePage interface
  *****************************************************************************/
 
-void SProjectLauncherPackagePage::Construct( const FArguments& InArgs, const FProjectLauncherModelRef& InModel )
+void SProjectLauncherPackagePage::Construct(const FArguments& InArgs, const TSharedRef<FProjectLauncherModel>& InModel)
 {
 	Model = InModel;
 
@@ -97,7 +99,7 @@ void SProjectLauncherPackagePage::Construct( const FArguments& InArgs, const FPr
 /* SProjectLauncherLaunchPage callbacks
  *****************************************************************************/
 
-FText SProjectLauncherPackagePage::HandlePackagingModeComboButtonContentText( ) const
+FText SProjectLauncherPackagePage::HandlePackagingModeComboButtonContentText() const
 {
 	ILauncherProfilePtr SelectedProfile = Model->GetSelectedProfile();
 
@@ -127,7 +129,7 @@ FText SProjectLauncherPackagePage::HandlePackagingModeComboButtonContentText( )
 }
 
 
-void SProjectLauncherPackagePage::HandlePackagingModeMenuEntryClicked( ELauncherProfilePackagingModes::Type PackagingMode )
+void SProjectLauncherPackagePage::HandlePackagingModeMenuEntryClicked(ELauncherProfilePackagingModes::Type PackagingMode)
 {
 	ILauncherProfilePtr SelectedProfile = Model->GetSelectedProfile();
 
@@ -141,8 +143,7 @@ void SProjectLauncherPackagePage::HandlePackagingModeMenuEntryClicked( ELauncher
 }
 
 
-
-EVisibility SProjectLauncherPackagePage::HandlePackagingSettingsAreaVisibility( ) const
+EVisibility SProjectLauncherPackagePage::HandlePackagingSettingsAreaVisibility() const
 {
 	ILauncherProfilePtr SelectedProfile = Model->GetSelectedProfile();
 
@@ -160,7 +161,7 @@ EVisibility SProjectLauncherPackagePage::HandlePackagingSettingsAreaVisibility(
 }
 
 
-void SProjectLauncherPackagePage::HandleProfileManagerProfileSelected( const ILauncherProfilePtr& SelectedProfile, const ILauncherProfilePtr& PreviousProfile )
+void SProjectLauncherPackagePage::HandleProfileManagerProfileSelected(const ILauncherProfilePtr& SelectedProfile, const ILauncherProfilePtr& PreviousProfile)
 {
 
 }
