Copying //UE4/Dev-Sequencer to //UE4/Main (Source: //UE4/Dev-Sequencer @ 3617642)

diff --git a/Engine/Source/Developer/ProjectLauncher/Private/Widgets/Package/SProjectLauncherPackagePage.h b/Engine/Source/Developer/ProjectLauncher/Private/Widgets/Package/SProjectLauncherPackagePage.h
--- a/Engine/Source/Developer/ProjectLauncher/Private/Widgets/Package/SProjectLauncherPackagePage.h
+++ b/Engine/Source/Developer/ProjectLauncher/Private/Widgets/Package/SProjectLauncherPackagePage.h
@@ -36,7 +36,7 @@ public:
 	 * @param InArgs The Slate argument list.
 	 * @param InModel The data model.
 	 */
-	void Construct( const FArguments& InArgs, const FProjectLauncherModelRef& InModel );
+	void Construct( const FArguments& InArgs, const TSharedRef<FProjectLauncherModel>& InModel );
 
 private:
 
@@ -55,7 +55,7 @@ private:
 private:
 
 	// Holds a pointer to the data model.
-	FProjectLauncherModelPtr Model;
+	TSharedPtr<FProjectLauncherModel> Model;
 
 	TSharedPtr<SProjectLauncherPackagingSettings> ProjectLauncherPackagingSettings;
 };
diff --git a/Engine/Source/Developer/ProjectLauncher/Private/Widgets/Package/SProjectLauncherPackagingSettings.cpp b/Engine/Source/Developer/ProjectLauncher/Private/Widgets/Package/SProjectLauncherPackagingSettings.cpp
--- a/Engine/Source/Developer/ProjectLauncher/Private/Widgets/Package/SProjectLauncherPackagingSettings.cpp
+++ b/Engine/Source/Developer/ProjectLauncher/Private/Widgets/Package/SProjectLauncherPackagingSettings.cpp
@@ -1,17 +1,18 @@
 // Copyright 1998-2017 Epic Games, Inc. All Rights Reserved.
 
-#include "Widgets/Package/SProjectLauncherPackagingSettings.h"
-#include "Widgets/SBoxPanel.h"
+#include "SProjectLauncherPackagingSettings.h"
+
+#include "DesktopPlatformModule.h"
+#include "EditorStyleSet.h"
+#include "Framework/Application/SlateApplication.h"
 #include "Styling/SlateTypes.h"
 #include "SlateOptMacros.h"
-#include "Framework/Application/SlateApplication.h"
-#include "Widgets/Layout/SBorder.h"
-#include "Widgets/Text/STextBlock.h"
+#include "Widgets/SBoxPanel.h"
 #include "Widgets/Input/SEditableTextBox.h"
 #include "Widgets/Input/SButton.h"
 #include "Widgets/Input/SCheckBox.h"
-#include "EditorStyleSet.h"
-#include "DesktopPlatformModule.h"
+#include "Widgets/Layout/SBorder.h"
+#include "Widgets/Text/STextBlock.h"
 
 
 #define LOCTEXT_NAMESPACE "SProjectLauncherPackagingSettings"
@@ -20,7 +21,7 @@
 /* SProjectLauncherPackagingSettings structors
  *****************************************************************************/
 
-SProjectLauncherPackagingSettings::~SProjectLauncherPackagingSettings( )
+SProjectLauncherPackagingSettings::~SProjectLauncherPackagingSettings()
 {
 	if (Model.IsValid())
 	{
@@ -33,7 +34,7 @@ SProjectLauncherPackagingSettings::~SProjectLauncherPackagingSettings( )
  *****************************************************************************/
 
 BEGIN_SLATE_FUNCTION_BUILD_OPTIMIZATION
-void SProjectLauncherPackagingSettings::Construct( const FArguments& InArgs, const FProjectLauncherModelRef& InModel )
+void SProjectLauncherPackagingSettings::Construct(const FArguments& InArgs, const TSharedRef<FProjectLauncherModel>& InModel)
 {
 	Model = InModel;
 
@@ -116,11 +117,13 @@ void SProjectLauncherPackagingSettings::Construct( const FArguments& InArgs, con
 }
 END_SLATE_FUNCTION_BUILD_OPTIMIZATION
 
+
 void SProjectLauncherPackagingSettings::UpdateDirectoryPathText()
 {
 	DirectoryPathTextBox->SetText(HandleDirectoryPathText());
 }
 
+
 /* SProjectLauncherPackagingSettings callbacks
  *****************************************************************************/
 
@@ -148,7 +151,6 @@ ECheckBoxState SProjectLauncherPackagingSettings::HandleForDistributionCheckBoxI
 }
 
 
-
 FText SProjectLauncherPackagingSettings::HandleDirectoryTitleText() const
 {
 	ILauncherProfilePtr SelectedProfile = Model->GetSelectedProfile();
@@ -168,6 +170,7 @@ FText SProjectLauncherPackagingSettings::HandleDirectoryTitleText() const
 	return FText::GetEmpty();
 }
 
+
 FText SProjectLauncherPackagingSettings::HandleDirectoryPathText() const
 {
 	ILauncherProfilePtr SelectedProfile = Model->GetSelectedProfile();
@@ -184,6 +187,7 @@ FText SProjectLauncherPackagingSettings::HandleDirectoryPathText() const
 	return FText::GetEmpty();
 }
 
+
 FText SProjectLauncherPackagingSettings::HandleHintPathText() const
 {
 	ILauncherProfilePtr SelectedProfile = Model->GetSelectedProfile();
@@ -197,11 +201,13 @@ FText SProjectLauncherPackagingSettings::HandleHintPathText() const
 	return FText::GetEmpty();
 }
 
-void SProjectLauncherPackagingSettings::HandleProfileManagerProfileSelected( const ILauncherProfilePtr& SelectedProfile, const ILauncherProfilePtr& PreviousProfile )
+
+void SProjectLauncherPackagingSettings::HandleProfileManagerProfileSelected(const ILauncherProfilePtr& SelectedProfile, const ILauncherProfilePtr& PreviousProfile)
 {
 	UpdateDirectoryPathText();
 }
 
+
 FReply SProjectLauncherPackagingSettings::HandleBrowseButtonClicked()
 {
 	IDesktopPlatform* DesktopPlatform = FDesktopPlatformModule::Get();
@@ -238,6 +244,7 @@ FReply SProjectLauncherPackagingSettings::HandleBrowseButtonClicked()
 	return FReply::Handled();
 }
 
+
 bool SProjectLauncherPackagingSettings::IsEditable() const
 {
 	ILauncherProfilePtr SelectedProfile = Model->GetSelectedProfile();
@@ -249,6 +256,7 @@ bool SProjectLauncherPackagingSettings::IsEditable() const
 	return false;
 }
 
+
 void SProjectLauncherPackagingSettings::OnTextChanged(const FText& InText)
 {
 	ILauncherProfilePtr SelectedProfile = Model->GetSelectedProfile();
@@ -259,6 +267,7 @@ void SProjectLauncherPackagingSettings::OnTextChanged(const FText& InText)
 	}
 }
 
+
 void SProjectLauncherPackagingSettings::OnTextCommitted(const FText& InText, ETextCommit::Type CommitInfo)
 {
 	if (CommitInfo == ETextCommit::OnEnter)
@@ -272,4 +281,5 @@ void SProjectLauncherPackagingSettings::OnTextCommitted(const FText& InText, ETe
 	}
 }
 
+
 #undef LOCTEXT_NAMESPACE
diff --git a/Engine/Source/Developer/ProjectLauncher/Private/Widgets/Profile/SProjectLauncherProfileListRow.h b/Engine/Source/Developer/ProjectLauncher/Private/Widgets/Profile/SProjectLauncherProfileListRow.h
--- a/Engine/Source/Developer/ProjectLauncher/Private/Widgets/Profile/SProjectLauncherProfileListRow.h
+++ b/Engine/Source/Developer/ProjectLauncher/Private/Widgets/Profile/SProjectLauncherProfileListRow.h
@@ -10,7 +10,7 @@
 #include "Widgets/Layout/SBorder.h"
 #include "Widgets/SBoxPanel.h"
 #include "Widgets/Views/STableRow.h"
-#include "Widgets/SProjectLauncherDelegates.h"
+#include "Widgets/Shared/ProjectLauncherDelegates.h"
 #include "EditorStyleSet.h"
 #include "Widgets/Views/SListView.h"
 #include "Widgets/Images/SImage.h"
@@ -55,7 +55,7 @@ public:
 	 * @param InArgs - The construction arguments.
 	 * @param InModel - The launcher model this list uses.
 	 */
-	void Construct(const FArguments& InArgs, const FProjectLauncherModelRef& InModel, const TSharedRef<STableViewBase>& InOwnerTableView)
+	void Construct(const FArguments& InArgs, const TSharedRef<FProjectLauncherModel>& InModel, const TSharedRef<STableViewBase>& InOwnerTableView)
 	{
 		STableRow<ILauncherProfilePtr>::ConstructInternal(
 			STableRow::FArguments()
@@ -184,7 +184,7 @@ private:
 private:
 
 	// Holds a pointer to the data model.
-	FProjectLauncherModelPtr Model;
+	TSharedPtr<FProjectLauncherModel> Model;
 
 	// Holds a reference to the launch profile that is displayed in this row.
 	ILauncherProfilePtr LaunchProfile;
diff --git a/Engine/Source/Developer/ProjectLauncher/Private/Widgets/Profile/SProjectLauncherProfileListView.cpp b/Engine/Source/Developer/ProjectLauncher/Private/Widgets/Profile/SProjectLauncherProfileListView.cpp
--- a/Engine/Source/Developer/ProjectLauncher/Private/Widgets/Profile/SProjectLauncherProfileListView.cpp
+++ b/Engine/Source/Developer/ProjectLauncher/Private/Widgets/Profile/SProjectLauncherProfileListView.cpp
@@ -1,13 +1,15 @@
 // Copyright 1998-2017 Epic Games, Inc. All Rights Reserved.
 
-#include "Widgets/Profile/SProjectLauncherProfileListView.h"
-#include "Misc/MessageDialog.h"
-#include "SlateOptMacros.h"
+#include "SProjectLauncherProfileListView.h"
+
 #include "Framework/Commands/UICommandList.h"
 #include "Framework/MultiBox/MultiBoxBuilder.h"
+#include "Misc/MessageDialog.h"
+#include "SlateOptMacros.h"
 #include "Widgets/Layout/SScrollBorder.h"
+
 #include "Models/ProjectLauncherCommands.h"
-#include "Widgets/Views/SListView.h"
+#include "Models/ProjectLauncherModel.h"
 #include "Widgets/Profile/SProjectLauncherProfileListRow.h"
 
 
@@ -20,7 +22,6 @@
 SProjectLauncherProfileListView::SProjectLauncherProfileListView()
 	: CommandList(MakeShareable(new FUICommandList))
 {
-
 }
 
 SProjectLauncherProfileListView::~SProjectLauncherProfileListView()
@@ -38,7 +39,7 @@ SProjectLauncherProfileListView::~SProjectLauncherProfileListView()
  *****************************************************************************/
 
 BEGIN_SLATE_FUNCTION_BUILD_OPTIMIZATION
-void SProjectLauncherProfileListView::Construct(const FArguments& InArgs, const FProjectLauncherModelRef& InModel)
+void SProjectLauncherProfileListView::Construct(const FArguments& InArgs, const TSharedRef<FProjectLauncherModel>& InModel)
 {
 	CreateCommands();
 
@@ -62,13 +63,13 @@ void SProjectLauncherProfileListView::Construct(const FArguments& InArgs, const
 		SNew(SVerticalBox)
 
 		+ SVerticalBox::Slot()
-		.FillHeight(1.0f)
-		[
-			SNew(SScrollBorder, LaunchProfileListView.ToSharedRef())
+			.FillHeight(1.0f)
 			[
-				LaunchProfileListView.ToSharedRef()
+				SNew(SScrollBorder, LaunchProfileListView.ToSharedRef())
+				[
+					LaunchProfileListView.ToSharedRef()
+				]
 			]
-		]
 	];
 
 	ProfileManager->OnProfileAdded().AddSP(this, &SProjectLauncherProfileListView::HandleProfileManagerProfileAdded);
@@ -76,6 +77,7 @@ void SProjectLauncherProfileListView::Construct(const FArguments& InArgs, const
 }
 END_SLATE_FUNCTION_BUILD_OPTIMIZATION
 
+
 void SProjectLauncherProfileListView::CreateCommands()
 {
 	const FProjectLauncherCommands& Commands = FProjectLauncherCommands::Get();
@@ -94,10 +96,11 @@ void SProjectLauncherProfileListView::CreateCommands()
 		FCanExecuteAction::CreateRaw(this, &SProjectLauncherProfileListView::HandleDeleteProfileCommandCanExecute));
 }
 
+
 /* SProjectLauncherDeployTargets implementation
  *****************************************************************************/
 
-void SProjectLauncherProfileListView::RefreshLaunchProfileList( )
+void SProjectLauncherProfileListView::RefreshLaunchProfileList()
 {
 	LaunchProfileListView->RequestListRefresh();
 }
@@ -111,6 +114,7 @@ bool SProjectLauncherProfileListView::HandleProfileRowIsEnabled(ILauncherProfile
 	return true;
 }
 
+
 FText SProjectLauncherProfileListView::HandleDeviceListRowToolTipText(ILauncherProfilePtr LaunchProfile) const
 {
 	return FText();
@@ -126,16 +130,19 @@ TSharedRef<ITableRow> SProjectLauncherProfileListView::HandleProfileListViewGene
 		.IsEnabled(this, &SProjectLauncherProfileListView::HandleProfileRowIsEnabled, InItem);
 }
 
+
 void SProjectLauncherProfileListView::HandleProfileManagerProfileAdded(const ILauncherProfileRef& AddedProfile)
 {
 	RefreshLaunchProfileList();
 }
 
+
 void SProjectLauncherProfileListView::HandleProfileManagerProfileRemoved(const ILauncherProfileRef& RemovedProfile)
 {
 	RefreshLaunchProfileList();
 }
 
+
 TSharedPtr<SWidget> SProjectLauncherProfileListView::MakeProfileContextMenu()
 {
 	FMenuBuilder MenuBuilder(/*bInShouldCloseWindowAfterMenuSelection=*/true, CommandList);
@@ -147,11 +154,13 @@ TSharedPtr<SWidget> SProjectLauncherProfileListView::MakeProfileContextMenu()
 	return MenuBuilder.MakeWidget();
 }
 
+
 bool SProjectLauncherProfileListView::HandleRenameProfileCommandCanExecute() const
 {
 	return true;
 }
 
+
 void SProjectLauncherProfileListView::HandleRenameProfileCommandExecute()
 {
 	TArray<ILauncherProfilePtr> ProfileList = LaunchProfileListView->GetSelectedItems();
@@ -169,11 +178,13 @@ void SProjectLauncherProfileListView::HandleRenameProfileCommandExecute()
 	}
 }
 
+
 bool SProjectLauncherProfileListView::HandleDuplicateProfileCommandCanExecute() const
 {
 	return true;
 }
 
+
 void SProjectLauncherProfileListView::HandleDuplicateProfileCommandExecute()
 {
 	TArray<ILauncherProfilePtr> ProfileList = LaunchProfileListView->GetSelectedItems();
@@ -184,11 +195,13 @@ void SProjectLauncherProfileListView::HandleDuplicateProfileCommandExecute()
 	}
 }
 
+
 bool SProjectLauncherProfileListView::HandleDeleteProfileCommandCanExecute() const
 {
 	return OnProfileDelete.IsBound();
 }
 
+
 void SProjectLauncherProfileListView::HandleDeleteProfileCommandExecute()
 {
 	if (!OnProfileDelete.IsBound())
diff --git a/Engine/Source/Developer/ProjectLauncher/Private/Widgets/Profile/SProjectLauncherProfileListView.h b/Engine/Source/Developer/ProjectLauncher/Private/Widgets/Profile/SProjectLauncherProfileListView.h
--- a/Engine/Source/Developer/ProjectLauncher/Private/Widgets/Profile/SProjectLauncherProfileListView.h
+++ b/Engine/Source/Developer/ProjectLauncher/Private/Widgets/Profile/SProjectLauncherProfileListView.h
@@ -2,15 +2,18 @@
 
 #pragma once
 
-#include "CoreMinimal.h"
+#include "Templates/SharedPointer.h"
 #include "Widgets/DeclarativeSyntaxSupport.h"
 #include "Widgets/SCompoundWidget.h"
-#include "Widgets/Views/STableViewBase.h"
-#include "Widgets/Views/STableRow.h"
-#include "Models/ProjectLauncherModel.h"
-#include "Widgets/SProjectLauncherDelegates.h"
+#include "Widgets/Views/SListView.h"
 
+#include "Widgets/Shared/ProjectLauncherDelegates.h"
+
+class FProjectLauncherModel;
 class FUICommandList;
+class ITableRow;
+class STableViewBase;
+
 
 /**
  * Implements the deployment targets panel.
@@ -38,92 +41,87 @@ public:
 	SLATE_END_ARGS()
 
 public:
-	/**
-	 * Constructor.
-	 */
+
+	/** Constructor. */
 	SProjectLauncherProfileListView();
 
-	/**
-	 * Destructor.
-	 */
-	~SProjectLauncherProfileListView( );
+	/** Destructor. */
+	~SProjectLauncherProfileListView();
 
 public:
 
 	/**
-	 * Constructs the widget.
+	 * Construct the widget.
 	 *
 	 * @param InArgs The Slate argument list.
 	 * @param InModel The data model.
 	 */
-	void Construct(const FArguments& InArgs, const FProjectLauncherModelRef& InModel);
+	void Construct(const FArguments& InArgs, const TSharedRef<FProjectLauncherModel>& InModel);
 
 protected:
 
-	/**
-	 * Refreshes the list of launch profiles.
-	 */
+	/** Refreshe the list of launch profiles. */
 	void RefreshLaunchProfileList();
 
 private:
 
-	// Builds the command list for the context menu on list items.
+	/** Builds the command list for the context menu on list items. */
 	void CreateCommands();
 
-	// Callback for getting the enabled state of a profile list row.
+	/** Callback for getting the enabled state of a profile list row. */
 	bool HandleProfileRowIsEnabled(ILauncherProfilePtr LaunchProfile) const;
 
-	// Callback for getting the tool tip text of a profile list row.
+	/** Callback for getting the tool tip text of a profile list row. */
 	FText HandleDeviceListRowToolTipText(ILauncherProfilePtr LaunchProfile) const;
 
-	// Callback for generating a row in the profile list view.
+	/** Callback for generating a row in the profile list view. */
 	TSharedRef<ITableRow> HandleProfileListViewGenerateRow(ILauncherProfilePtr InItem, const TSharedRef<STableViewBase>& OwnerTable) const;
 
-	// Callback for when a launch profile has been added to the profile manager.
+	/** Callback for when a launch profile has been added to the profile manager. */
 	void HandleProfileManagerProfileAdded(const ILauncherProfileRef& AddedProfile);
 
-	// Callback for when a launch profile has been removed to the profile manager.
+	/** Callback for when a launch profile has been removed to the profile manager. */
 	void HandleProfileManagerProfileRemoved(const ILauncherProfileRef& RemovedProfile);
 
-	// Callback for creating a context menu for the list.
+	/** Callback for creating a context menu for the list. */
 	TSharedPtr<SWidget> MakeProfileContextMenu();
 
-	// Callback returns true if the rename command can be executed.
+	/** Callback returns true if the rename command can be executed. */
 	bool HandleRenameProfileCommandCanExecute() const;
 
-	// Callback to execute the rename command from the context menu.
+	/** Callback to execute the rename command from the context menu. */
 	void HandleRenameProfileCommandExecute();
 
-	// Callback returns true if the duplicate command can be executed.
+	/** Callback returns true if the duplicate command can be executed. */
 	bool HandleDuplicateProfileCommandCanExecute() const;
 
-	// Callback to execute the duplicate command from the context menu.
+	/** Callback to execute the duplicate command from the context menu. */
 	void HandleDuplicateProfileCommandExecute();
 
-	// Callback returns true if the delete command can be executed.
+	/** Callback returns true if the delete command can be executed. */
 	bool HandleDeleteProfileCommandCanExecute() const;
 
-	// Callback to execute the delete command from the context menu.
+	/** Callback to execute the delete command from the context menu. */
 	void HandleDeleteProfileCommandExecute();
 
 private:
 
-	// Holds the launch profile list view .
+	/** The launch profile list view. */
 	TSharedPtr<SListView<ILauncherProfilePtr> > LaunchProfileListView;
 
-	// Holds a pointer to the data model.
-	FProjectLauncherModelPtr Model;
+	/** Pointer to the data model. */
+	TSharedPtr<FProjectLauncherModel> Model;
 
-	// Holds a delegate to be invoked when a profile is to be edited.
+	/** A delegate to be invoked when a profile is to be edited. */
 	FOnProfileRun OnProfileEdit;
 
-	// Holds a delegate to be invoked when a profile is run.
+	/** A delegate to be invoked when a profile is run. */
 	FOnProfileRun OnProfileRun;
 
-	// Holds a delegate to be invoked when a profile is deleted.
+	/** A delegate to be invoked when a profile is deleted. */
 	FOnProfileRun OnProfileDelete;
 
-	/** Commands handled by this widget */
+	/** Commands handled by this widget. */
 	TSharedRef<FUICommandList> CommandList;
 
 };
diff --git a/Engine/Source/Developer/ProjectLauncher/Private/Widgets/Project/SProjectLauncherProjectPage.cpp b/Engine/Source/Developer/ProjectLauncher/Private/Widgets/Project/SProjectLauncherProjectPage.cpp
--- a/Engine/Source/Developer/ProjectLauncher/Private/Widgets/Project/SProjectLauncherProjectPage.cpp
+++ b/Engine/Source/Developer/ProjectLauncher/Private/Widgets/Project/SProjectLauncherProjectPage.cpp
@@ -1,8 +1,10 @@
 // Copyright 1998-2017 Epic Games, Inc. All Rights Reserved.
 
-#include "Widgets/Project/SProjectLauncherProjectPage.h"
+#include "SProjectLauncherProjectPage.h"
+
 #include "Widgets/SBoxPanel.h"
 #include "Widgets/Text/STextBlock.h"
+
 #include "Widgets/Project/SProjectLauncherProjectPicker.h"
 
 
@@ -12,7 +14,7 @@
 /* SProjectLauncherProjectPage interface
  *****************************************************************************/
 
-void SProjectLauncherProjectPage::Construct( const FArguments& InArgs, const FProjectLauncherModelRef& InModel, bool InShowConfig )
+void SProjectLauncherProjectPage::Construct(const FArguments& InArgs, const TSharedRef<FProjectLauncherModel>& InModel, bool InShowConfig)
 {
 	Model = InModel;
 
diff --git a/Engine/Source/Developer/ProjectLauncher/Private/Widgets/Project/SProjectLauncherProjectPage.h b/Engine/Source/Developer/ProjectLauncher/Private/Widgets/Project/SProjectLauncherProjectPage.h
--- a/Engine/Source/Developer/ProjectLauncher/Private/Widgets/Project/SProjectLauncherProjectPage.h
+++ b/Engine/Source/Developer/ProjectLauncher/Private/Widgets/Project/SProjectLauncherProjectPage.h
@@ -27,10 +27,10 @@ public:
 	 * @param InArgs The Slate argument list.
 	 * @param InModel The data model.
 	 */
-	void Construct(	const FArguments& InArgs, const FProjectLauncherModelRef& InModel, bool InShowConfig = true );
+	void Construct(	const FArguments& InArgs, const TSharedRef<FProjectLauncherModel>& InModel, bool InShowConfig = true );
 
 private:
 
 	// Holds a pointer to the data model.
-	FProjectLauncherModelPtr Model;
+	TSharedPtr<FProjectLauncherModel> Model;
 };
diff --git a/Engine/Source/Developer/ProjectLauncher/Private/Widgets/Project/SProjectLauncherProjectPicker.cpp b/Engine/Source/Developer/ProjectLauncher/Private/Widgets/Project/SProjectLauncherProjectPicker.cpp
--- a/Engine/Source/Developer/ProjectLauncher/Private/Widgets/Project/SProjectLauncherProjectPicker.cpp
+++ b/Engine/Source/Developer/ProjectLauncher/Private/Widgets/Project/SProjectLauncherProjectPicker.cpp
@@ -1,16 +1,18 @@
 // Copyright 1998-2017 Epic Games, Inc. All Rights Reserved.
 
-#include "Widgets/Project/SProjectLauncherProjectPicker.h"
-#include "Widgets/SBoxPanel.h"
+#include "SProjectLauncherProjectPicker.h"
+
+#include "DesktopPlatformModule.h"
+#include "EditorStyleSet.h"
 #include "Framework/Application/SlateApplication.h"
-#include "Textures/SlateIcon.h"
 #include "Framework/Commands/UIAction.h"
-#include "Widgets/Layout/SBorder.h"
-#include "Widgets/Text/STextBlock.h"
 #include "Framework/MultiBox/MultiBoxBuilder.h"
+#include "Textures/SlateIcon.h"
+#include "Widgets/SBoxPanel.h"
 #include "Widgets/Input/SComboButton.h"
-#include "EditorStyleSet.h"
-#include "DesktopPlatformModule.h"
+#include "Widgets/Layout/SBorder.h"
+#include "Widgets/Text/STextBlock.h"
+
 #include "Widgets/Shared/SProjectLauncherFormLabel.h"
 
 
@@ -20,7 +22,7 @@
 /* SProjectLauncherProjectPicker structors
  *****************************************************************************/
 
-SProjectLauncherProjectPicker::~SProjectLauncherProjectPicker( )
+SProjectLauncherProjectPicker::~SProjectLauncherProjectPicker()
 {
 	if (Model.IsValid())
 	{
@@ -32,7 +34,7 @@ SProjectLauncherProjectPicker::~SProjectLauncherProjectPicker( )
 /* SProjectLauncherProjectPicker interface
  *****************************************************************************/
 
-void SProjectLauncherProjectPicker::Construct( const FArguments& InArgs, const FProjectLauncherModelRef& InModel)
+void SProjectLauncherProjectPicker::Construct(const FArguments& InArgs, const TSharedRef<FProjectLauncherModel>& InModel)
 {
 	Model = InModel;
 
@@ -48,7 +50,7 @@ void SProjectLauncherProjectPicker::Construct( const FArguments& InArgs, const F
 /* SProjectLauncherProjectPicker implementation
  *****************************************************************************/
 
-TSharedRef<SWidget> SProjectLauncherProjectPicker::MakeProjectMenuWidget( )
+TSharedRef<SWidget> SProjectLauncherProjectPicker::MakeProjectMenuWidget()
 {
 	FMenuBuilder MenuBuilder(true, NULL);
 	{
@@ -63,7 +65,7 @@ TSharedRef<SWidget> SProjectLauncherProjectPicker::MakeProjectMenuWidget( )
 		{
 			FString ProjectPath = FPaths::RootDir() / AvailableGames[Index] / AvailableGames[Index] + TEXT(".uproject");
 			FUIAction ProjectAction(FExecuteAction::CreateSP(this, &SProjectLauncherProjectPicker::HandleProjectMenuEntryClicked, ProjectPath));
-			MenuBuilder.AddMenuEntry( FText::FromString( AvailableGames[Index] ), FText::FromString( ProjectPath ), FSlateIcon(), ProjectAction);
+			MenuBuilder.AddMenuEntry(FText::FromString(AvailableGames[Index]), FText::FromString(ProjectPath), FSlateIcon(), ProjectAction);
 		}
 
 		MenuBuilder.AddMenuSeparator();
@@ -119,7 +121,7 @@ TSharedRef<SWidget> SProjectLauncherProjectPicker::MakeProjectWidget()
 /* SProjectLauncherProjectPicker callbacks
  *****************************************************************************/
 
-FText SProjectLauncherProjectPicker::HandleProjectComboButtonText( ) const
+FText SProjectLauncherProjectPicker::HandleProjectComboButtonText() const
 {
 	if (LaunchProfileAttr.IsBound())
 	{
@@ -141,7 +143,7 @@ FText SProjectLauncherProjectPicker::HandleProjectComboButtonText( ) const
 }
 
 
-FText SProjectLauncherProjectPicker::HandleProjectComboButtonToolTip( ) const
+FText SProjectLauncherProjectPicker::HandleProjectComboButtonToolTip() const
 {
 	return LOCTEXT("SelectProjectText_Tooltip", "Select or browse for a project");
 }
@@ -160,7 +162,7 @@ void SProjectLauncherProjectPicker::HandleAnyProjectClicked(FString ProjectPath)
 }
 
 
-void SProjectLauncherProjectPicker::HandleProjectMenuEntryClicked( FString ProjectPath )
+void SProjectLauncherProjectPicker::HandleProjectMenuEntryClicked(FString ProjectPath)
 {
 	if (ProjectPath.IsEmpty())
 	{
diff --git a/Engine/Source/Developer/ProjectLauncher/Private/Widgets/Project/SProjectLauncherProjectPicker.h b/Engine/Source/Developer/ProjectLauncher/Private/Widgets/Project/SProjectLauncherProjectPicker.h
--- a/Engine/Source/Developer/ProjectLauncher/Private/Widgets/Project/SProjectLauncherProjectPicker.h
+++ b/Engine/Source/Developer/ProjectLauncher/Private/Widgets/Project/SProjectLauncherProjectPicker.h
@@ -39,7 +39,7 @@ public:
 	 * @param InArgs The Slate argument list.
 	 * @param InModel The data model.
 	 */
-	void Construct(const FArguments& InArgs, const FProjectLauncherModelRef& InModel);
+	void Construct(const FArguments& InArgs, const TSharedRef<FProjectLauncherModel>& InModel);
 
 protected:
 
@@ -86,6 +86,6 @@ private:
 	TArray<TSharedPtr<FString>> ProjectList;
 
 	// Holds a pointer to the data model.
-	FProjectLauncherModelPtr Model;
+	TSharedPtr<FProjectLauncherModel> Model;
 
 };
diff --git a/Engine/Source/Developer/ProjectLauncher/Private/Widgets/SProjectLauncher.cpp b/Engine/Source/Developer/ProjectLauncher/Private/Widgets/SProjectLauncher.cpp
--- a/Engine/Source/Developer/ProjectLauncher/Private/Widgets/SProjectLauncher.cpp
+++ b/Engine/Source/Developer/ProjectLauncher/Private/Widgets/SProjectLauncher.cpp
@@ -1,28 +1,30 @@
 // Copyright 1998-2017 Epic Games, Inc. All Rights Reserved.
 
-#include "Widgets/SProjectLauncher.h"
-#include "Widgets/SBoxPanel.h"
-#include "Styling/SlateTypes.h"
-#include "Styling/CoreStyle.h"
-#include "SlateOptMacros.h"
+#include "SProjectLauncher.h"
+
+#include "EditorStyleSet.h"
 #include "Framework/Commands/UIAction.h"
-#include "Widgets/Layout/SBorder.h"
-#include "Widgets/Images/SImage.h"
-#include "Widgets/Text/STextBlock.h"
+#include "Framework/Docking/TabManager.h"
 #include "Framework/MultiBox/MultiBoxBuilder.h"
+#include "SlateOptMacros.h"
+#include "Styling/CoreStyle.h"
+#include "Styling/SlateTypes.h"
+#include "Widgets/SBoxPanel.h"
+#include "Widgets/Images/SImage.h"
 #include "Widgets/Input/SButton.h"
+#include "Widgets/Input/SCheckBox.h"
 #include "Widgets/Input/SComboButton.h"
+#include "Widgets/Layout/SBorder.h"
 #include "Widgets/Layout/SSplitter.h"
-#include "Widgets/Input/SCheckBox.h"
-#include "Framework/Docking/TabManager.h"
-#include "EditorStyleSet.h"
+#include "Widgets/Layout/SWidgetSwitcher.h"
+#include "Widgets/Text/STextBlock.h"
+
 #include "Models/ProjectLauncherCommands.h"
-#include "Widgets/Project/SProjectLauncherProjectPicker.h"
 #include "Widgets/Deploy/SProjectLauncherSimpleDeviceListView.h"
 #include "Widgets/Profile/SProjectLauncherProfileListView.h"
-#include "Widgets/SProjectLauncherProgress.h"
-#include "Widgets/SProjectLauncherSettings.h"
-#include "Widgets/Layout/SWidgetSwitcher.h"
+#include "Widgets/Progress/SProjectLauncherProgress.h"
+#include "Widgets/Project/SProjectLauncherProjectPicker.h"
+#include "Widgets/Settings/SProjectLauncherSettings.h"
 
 
 #define LOCTEXT_NAMESPACE "SProjectLauncher"
@@ -40,7 +42,8 @@ SProjectLauncher::SProjectLauncher()
 	}
 }
 
-SProjectLauncher::~SProjectLauncher( )
+
+SProjectLauncher::~SProjectLauncher()
 {
 	if (GConfig != NULL)
 	{
@@ -59,7 +62,7 @@ SProjectLauncher::~SProjectLauncher( )
  *****************************************************************************/
 
 BEGIN_SLATE_FUNCTION_BUILD_OPTIMIZATION
-void SProjectLauncher::Construct( const FArguments& InArgs, const TSharedRef<SDockTab>& ConstructUnderMajorTab, const TSharedPtr<SWindow>& ConstructUnderWindow, const FProjectLauncherModelRef& InModel )
+void SProjectLauncher::Construct(const FArguments& InArgs, const TSharedRef<SDockTab>& ConstructUnderMajorTab, const TSharedPtr<SWindow>& ConstructUnderWindow, const TSharedRef<FProjectLauncherModel>& InModel)
 {
 	FProjectLauncherCommands::Register();
 
@@ -263,10 +266,11 @@ void SProjectLauncher::Construct( const FArguments& InArgs, const TSharedRef<SDo
 }
 END_SLATE_FUNCTION_BUILD_OPTIMIZATION
 
+
 /* SProjectLauncher implementation
 *****************************************************************************/
 
-void SProjectLauncher::FillWindowMenu( FMenuBuilder& MenuBuilder, TSharedRef<FWorkspaceItem> RootMenuGroup )
+void SProjectLauncher::FillWindowMenu(FMenuBuilder& MenuBuilder, TSharedRef<FWorkspaceItem> RootMenuGroup)
 {
 #if !WITH_EDITOR
 	MenuBuilder.BeginSection("WindowGlobalTabSpawners", LOCTEXT("UfeMenuGroup", "Unreal Frontend"));
@@ -277,6 +281,7 @@ void SProjectLauncher::FillWindowMenu( FMenuBuilder& MenuBuilder, TSharedRef<FWo
 #endif //!WITH_EDITOR
 }
 
+
 /* SProjectLauncher callbacks
 *****************************************************************************/
 
@@ -285,27 +290,32 @@ void SProjectLauncher::OnAdvancedChanged(const ECheckBoxState NewCheckedState)
 	bAdvanced = (NewCheckedState == ECheckBoxState::Checked);
 }
 
+
 ECheckBoxState SProjectLauncher::OnIsAdvanced() const
 {
 	return (bAdvanced) ? ECheckBoxState::Checked : ECheckBoxState::Unchecked;
 }
 
+
 const FSlateBrush* SProjectLauncher::GetAdvancedToggleBrush() const
 {
 	return FEditorStyle::GetBrush("LauncherCommand.AdvancedBuild.Medium");
 }
 
+
 bool SProjectLauncher::GetIsAdvanced() const
 {
 	return bAdvanced;
 }
 
+
 void SProjectLauncher::OnProfileEdit(const ILauncherProfileRef& Profile)
 {
 	Model->SelectProfile(Profile);
 	WidgetSwitcher->SetActiveWidgetIndex(ELauncherPanels::ProfileEditor);
 }
 
+
 void SProjectLauncher::OnProfileRun(const ILauncherProfileRef& Profile)
 {
 	LauncherProfile = Profile;
@@ -318,11 +328,13 @@ void SProjectLauncher::OnProfileRun(const ILauncherProfileRef& Profile)
 	}
 }
 
+
 void SProjectLauncher::OnProfileDelete(const ILauncherProfileRef& Profile)
 {
 	Model->GetProfileManager()->RemoveProfile(Profile);
 }
 
+
 FReply SProjectLauncher::OnAddCustomLaunchProfileClicked()
 {
 	ILauncherProfileRef Profile = Model->GetProfileManager()->AddNewProfile();
@@ -334,11 +346,13 @@ FReply SProjectLauncher::OnAddCustomLaunchProfileClicked()
 	return FReply::Handled();
 }
 
+
 EVisibility SProjectLauncher::GetProfileWizardsMenuVisibility() const
 {
 	return (Model->GetProfileManager()->GetProfileWizards().Num() > 0) ? EVisibility::Visible : EVisibility::Collapsed;
 }
 
+
 TSharedRef<SWidget> SProjectLauncher::MakeProfileWizardsMenu()
 {
 	FMenuBuilder MenuBuilder(true, NULL);
@@ -363,11 +377,13 @@ TSharedRef<SWidget> SProjectLauncher::MakeProfileWizardsMenu()
 	return MenuBuilder.MakeWidget();
 }
 
+
 void SProjectLauncher::ExecProfileWizard(ILauncherProfileWizardPtr InWizard)
 {
 	InWizard->HandleCreateLauncherProfile(Model->GetProfileManager());
 }
 
+
 FReply SProjectLauncher::OnProfileSettingsClose()
 {
 	if (LauncherWorker.IsValid())
@@ -381,6 +397,7 @@ FReply SProjectLauncher::OnProfileSettingsClose()
 	return FReply::Handled();
 }
 
+
 FReply SProjectLauncher::OnProgressClose()
 {
 	if (LauncherWorker.IsValid())
@@ -394,6 +411,7 @@ FReply SProjectLauncher::OnProgressClose()
 	return FReply::Handled();
 }
 
+
 FReply SProjectLauncher::OnRerunClicked()
 {
 	if (LauncherWorker.IsValid())
@@ -410,4 +428,5 @@ FReply SProjectLauncher::OnRerunClicked()
 	return FReply::Handled();
 }
 
+
 #undef LOCTEXT_NAMESPACE
diff --git a/Engine/Source/Developer/ProjectLauncher/Private/Widgets/SProjectLauncher.h b/Engine/Source/Developer/ProjectLauncher/Private/Widgets/SProjectLauncher.h
--- a/Engine/Source/Developer/ProjectLauncher/Private/Widgets/SProjectLauncher.h
+++ b/Engine/Source/Developer/ProjectLauncher/Private/Widgets/SProjectLauncher.h
@@ -2,14 +2,15 @@
 
 #pragma once
 
-#include "CoreMinimal.h"
-#include "SlateFwd.h"
+#include "CoreTypes.h"
+#include "Framework/Docking/WorkspaceItem.h"
+#include "Input/Reply.h"
 #include "Layout/Visibility.h"
+#include "SlateFwd.h"
+#include "Templates/SharedPointer.h"
 #include "Widgets/DeclarativeSyntaxSupport.h"
-#include "Input/Reply.h"
-#include "Widgets/SWidget.h"
 #include "Widgets/SCompoundWidget.h"
-#include "Framework/Docking/WorkspaceItem.h"
+
 #include "Models/ProjectLauncherModel.h"
 
 class FMenuBuilder;
@@ -18,9 +19,12 @@ class SProjectLauncherProgress;
 class SProjectLauncherSettings;
 class SProjectLauncherSimpleDeviceListView;
 class SWindow;
+
 struct FSlateBrush;
+
 enum class ECheckBoxState : uint8;
 
+
 /**
  * Implements a Slate widget for the launcher user interface.
  */
@@ -29,7 +33,7 @@ class SProjectLauncher
 {
 public:
 
-	SLATE_BEGIN_ARGS(SProjectLauncher){ }
+	SLATE_BEGIN_ARGS(SProjectLauncher) { }
 
 		/** Exposes a delegate to be invoked when the launcher has closed. */
 		SLATE_EVENT(FSimpleDelegate, OnClosed)
@@ -38,15 +42,11 @@ public:
 
 public:
 
-	/**
-	* Constructor.
-	*/
+	/** Constructor. */
 	SProjectLauncher();
 
-	/**
-	 * Destructor.
-	 */
-	~SProjectLauncher( );
+	/** Destructor. */
+	~SProjectLauncher();
 
 public:
 
@@ -58,9 +58,7 @@ public:
 	 * @param ConstructUnderWindow The window in which this widget is being constructed.
 	 * @param InModel The view model to use.
 	 */
-	void Construct( const FArguments& InArgs, const TSharedRef<SDockTab>& ConstructUnderMajorTab, const TSharedPtr<SWindow>& ConstructUnderWindow, const FProjectLauncherModelRef& InModel );
-
-	void ShowWidget(int32 Index);
+	void Construct(const FArguments& InArgs, const TSharedRef<SDockTab>& ConstructUnderMajorTab, const TSharedPtr<SWindow>& ConstructUnderWindow, const TSharedRef<FProjectLauncherModel>& InModel);
 
 protected:
 
@@ -70,85 +68,80 @@ protected:
 	 * @param MenuBuilder The multi-box builder that should be filled with content for this pull-down menu.
 	 * @param RootMenuGroup The root menu group.
 	 */
-	static void FillWindowMenu( FMenuBuilder& MenuBuilder, TSharedRef<FWorkspaceItem> RootMenuGroup );
-
-	/**
-	 * Launches the selected profile.
-	 */
-	void LaunchSelectedProfile( );
+	static void FillWindowMenu(FMenuBuilder& MenuBuilder, TSharedRef<FWorkspaceItem> RootMenuGroup);
 
 private:
 
-	// Callback for toggling the advanced mode.
+	/** Callback for toggling the advanced mode. */
 	void OnAdvancedChanged(const ECheckBoxState NewCheckedState);
 
-	// Callback to determine if we are in advanced mode.
+	/** Callback to determine if we are in advanced mode. */
 	ECheckBoxState OnIsAdvanced() const;
 
-	// Get advanced toggle button brush.
+	/** Get advanced toggle button brush. */
 	const FSlateBrush* GetAdvancedToggleBrush() const;
 
-	// Callback for whether advanced options should be shown.
+	/** Callback for whether advanced options should be shown. */
 	bool GetIsAdvanced() const;
 
-	// Callback for editing a profile
+	/** Callback for editing a profile. */
 	void OnProfileEdit(const ILauncherProfileRef& Profile);
 
-	// Callback for running a profile
+	/** Callback for running a profile. */
 	void OnProfileRun(const ILauncherProfileRef& Profile);
 
-	// Callback for deleting a profile
+	/** Callback for deleting a profile. */
 	void OnProfileDelete(const ILauncherProfileRef& Profile);
 
-	// Callback for clicking the Add New Launch Profile
+	/** Callback for clicking the Add New Launch Profile. */
 	FReply OnAddCustomLaunchProfileClicked();
 
-	// Profile wizard menu visibility (hidden if there are no registered wizards)
+	/** Profile wizard menu visibility (hidden if there are no registered wizards). */
 	EVisibility GetProfileWizardsMenuVisibility() const;
 
-	// Callback for filling profile wizard menu
+	/** Callback for filling profile wizard menu. */
 	TSharedRef<SWidget> MakeProfileWizardsMenu();
 
-	// Execute specified profile wizard
+	/** Execute specified profile wizard. */
 	void ExecProfileWizard(ILauncherProfileWizardPtr InWizard);
 
-	// Callback for when the settings panel is closed
+	/** Callback for when the settings panel is closed. */
 	FReply OnProfileSettingsClose();
 
-	// Callback for when the progress panel is closed
+	/** Callback for when the progress panel is closed. */
 	FReply OnProgressClose();
 
-	// Callback for when the progress panel requests the profile to be re run
+	/** Callback for when the progress panel requests the profile to be re run. */
 	FReply OnRerunClicked();
 
 private:
 
-	// Holds the current launcher worker, if any.
+	/** The current launcher worker, if any. */
 	ILauncherWorkerPtr LauncherWorker;
 
-	// Holds the Launcher profile the launcherWorker is running.
+	/** The Launcher profile the launcherWorker is running. */
 	ILauncherProfilePtr LauncherProfile;
 
-	// Holds a pointer to the view model.
-	FProjectLauncherModelPtr Model;
+	/** Holds a pointer to the view model. */
+	TSharedPtr<FProjectLauncherModel> Model;
 
-	// Holds the profile settings panel.
+	/** The profile settings panel. */
 	TSharedPtr<SProjectLauncherSettings> ProfileSettingsPanel;
 
-	// Holds the progress panel.
+	/** The progress panel. */
 	TSharedPtr<SProjectLauncherProgress> ProgressPanel;
 
-	// Holds the widget switcher.
+	/** The widget switcher. */
 	TSharedPtr<SWidgetSwitcher> WidgetSwitcher;
 
-	// Contains the launch list widgets.
-	//TSharedPtr<SBorder> LaunchList;
+	/** Contains the launch list widgets. */
+//	TSharedPtr<SBorder> LaunchList;
 
 	TSharedPtr<SProjectLauncherSimpleDeviceListView> LaunchList;
 
-	// Contains the profile list widgets.
+	/** Contains the profile list widgets. */
 	TSharedPtr<SBorder> ProfileList;
 
-	// Whether we are showing advanced options.
+	/** Whether we are showing advanced options. */
 	bool bAdvanced;
 };
diff --git a/Engine/Source/Developer/ProjectLauncher/Private/Widgets/Shared/ProjectLauncherDelegates.h b/Engine/Source/Developer/ProjectLauncher/Private/Widgets/Shared/ProjectLauncherDelegates.h
new file mode 100644
--- /dev/null
+++ b/Engine/Source/Developer/ProjectLauncher/Private/Widgets/Shared/ProjectLauncherDelegates.h
@@ -0,0 +1,14 @@
+// Copyright 1998-2017 Epic Games, Inc. All Rights Reserved.
+
+#pragma once
+
+#include "Delegates/Delegate.h"
+#include "ILauncherProfile.h"
+
+
+/**
+ * Delegate type for SProjectLauncher Profile running.
+ *
+ * The first parameter is the SProjectLauncher Profile that you want to run.
+ */
+DECLARE_DELEGATE_OneParam(FOnProfileRun, const ILauncherProfileRef&)
diff --git a/Engine/Source/Developer/ProjectLauncher/Private/Widgets/Shared/SProjectLauncherCookModeSelector.h b/Engine/Source/Developer/ProjectLauncher/Private/Widgets/Shared/SProjectLauncherCookModeSelector.h
--- a/Engine/Source/Developer/ProjectLauncher/Private/Widgets/Shared/SProjectLauncherCookModeSelector.h
+++ b/Engine/Source/Developer/ProjectLauncher/Private/Widgets/Shared/SProjectLauncherCookModeSelector.h
@@ -7,7 +7,7 @@
 #include "Layout/Margin.h"
 #include "Widgets/SCompoundWidget.h"
 #include "Styling/CoreStyle.h"
-#include "Interfaces/ILauncherProfile.h"
+#include "ILauncherProfile.h"
 #include "Widgets/Text/STextBlock.h"
 #include "Textures/SlateIcon.h"
 #include "Framework/Commands/UIAction.h"
diff --git a/Engine/Source/Developer/ProjectLauncher/Private/Widgets/Shared/SProjectLauncherInstanceTypeSelector.h b/Engine/Source/Developer/ProjectLauncher/Private/Widgets/Shared/SProjectLauncherInstanceTypeSelector.h
--- a/Engine/Source/Developer/ProjectLauncher/Private/Widgets/Shared/SProjectLauncherInstanceTypeSelector.h
+++ b/Engine/Source/Developer/ProjectLauncher/Private/Widgets/Shared/SProjectLauncherInstanceTypeSelector.h
@@ -4,9 +4,9 @@
 
 #include "CoreMinimal.h"
 #include "Delegates/Delegate.h"
+#include "ILauncherProfileLaunchRole.h"
 #include "Internationalization/Text.h"
 #include "Internationalization/Internationalization.h"
-#include "Interfaces/ILauncherProfileLaunchRole.h"
 #include "Layout/Margin.h"
 #include "Widgets/DeclarativeSyntaxSupport.h"
 #include "Widgets/SCompoundWidget.h"
diff --git a/Engine/Source/Developer/ProjectLauncher/Private/Widgets/Shared/SProjectLauncherProfileLaunchButton.h b/Engine/Source/Developer/ProjectLauncher/Private/Widgets/Shared/SProjectLauncherProfileLaunchButton.h
--- a/Engine/Source/Developer/ProjectLauncher/Private/Widgets/Shared/SProjectLauncherProfileLaunchButton.h
+++ b/Engine/Source/Developer/ProjectLauncher/Private/Widgets/Shared/SProjectLauncherProfileLaunchButton.h
@@ -16,10 +16,12 @@
 #include "Widgets/Input/SButton.h"
 #include "Widgets/SOverlay.h"
 #include "Widgets/SToolTip.h"
-#include "Widgets/SProjectLauncherValidation.h"
+#include "Widgets/Shared/SProjectLauncherValidation.h"
+
 
 #define LOCTEXT_NAMESPACE "SProjectLauncherProfileLaunchButton"
 
+
 /**
  * Implements a build configuration selector widget.
  */
@@ -164,10 +166,10 @@ private:
 
 private:
 
-	// Attribute for the launch profile this widget launches. 
+	/** Attribute for the launch profile this widget launches. */
 	TAttribute<ILauncherProfilePtr> LaunchProfileAttr;
 
-	// Holds a pointer to our custom tooltip.
+	/** Holds a pointer to our custom tooltip. */
 	TSharedPtr<SToolTip> ErrorToolTipWidget;
 
 };
diff --git a/Engine/Source/Developer/ProjectLauncher/Private/Widgets/Shared/SProjectLauncherProfileNameDescEditor.cpp b/Engine/Source/Developer/ProjectLauncher/Private/Widgets/Shared/SProjectLauncherProfileNameDescEditor.cpp
--- a/Engine/Source/Developer/ProjectLauncher/Private/Widgets/Shared/SProjectLauncherProfileNameDescEditor.cpp
+++ b/Engine/Source/Developer/ProjectLauncher/Private/Widgets/Shared/SProjectLauncherProfileNameDescEditor.cpp
@@ -1,8 +1,9 @@
 // Copyright 1998-2017 Epic Games, Inc. All Rights Reserved.
 
-#include "Widgets/Shared/SProjectLauncherProfileNameDescEditor.h"
-#include "Widgets/SBoxPanel.h"
+#include "SProjectLauncherProfileNameDescEditor.h"
+
 #include "Styling/CoreStyle.h"
+#include "Widgets/SBoxPanel.h"
 #include "Widgets/Images/SImage.h"
 #include "Widgets/Layout/SBox.h"
 #include "Widgets/Text/SInlineEditableTextBlock.h"
@@ -11,8 +12,7 @@
 #define LOCTEXT_NAMESPACE "SProjectLauncherProfileNameDescEditor"
 
 
-
-void SProjectLauncherProfileNameDescEditor::Construct(const FArguments& InArgs, const FProjectLauncherModelRef& InModel, bool InShowAddDescriptionText)
+void SProjectLauncherProfileNameDescEditor::Construct(const FArguments& InArgs, const TSharedRef<FProjectLauncherModel>& InModel, bool InShowAddDescriptionText)
 {
 	EnterTextDescription = FText(LOCTEXT("LaunchProfileEnterDescription", "Enter a description here."));
 
@@ -67,6 +67,7 @@ void SProjectLauncherProfileNameDescEditor::Construct(const FArguments& InArgs,
 		];
 }
 
+
 void SProjectLauncherProfileNameDescEditor::TriggerNameEdit()
 {
 	if (NameEditableTextBlock.IsValid())
diff --git a/Engine/Source/Developer/ProjectLauncher/Private/Widgets/Shared/SProjectLauncherProfileNameDescEditor.h b/Engine/Source/Developer/ProjectLauncher/Private/Widgets/Shared/SProjectLauncherProfileNameDescEditor.h
--- a/Engine/Source/Developer/ProjectLauncher/Private/Widgets/Shared/SProjectLauncherProfileNameDescEditor.h
+++ b/Engine/Source/Developer/ProjectLauncher/Private/Widgets/Shared/SProjectLauncherProfileNameDescEditor.h
@@ -31,7 +31,7 @@ public:
 	 * @param InArgs The Slate argument list.
 	 * @param InModel The data model.
 	 */
-	void Construct(const FArguments& InArgs, const FProjectLauncherModelRef& InModel, bool InShowAddDescriptionText);
+	void Construct(const FArguments& InArgs, const TSharedRef<FProjectLauncherModel>& InModel, bool InShowAddDescriptionText);
 
 	/**
 	 * Triggers a name edit for the profile.
@@ -104,7 +104,7 @@ private:
 private:
 
 	/** Holds a pointer to the data model. */
-	FProjectLauncherModelPtr Model;
+	TSharedPtr<FProjectLauncherModel> Model;
 
 	// Attribute for the launch profile this widget edits. 
 	TAttribute<ILauncherProfilePtr> LaunchProfileAttr;
diff --git a/Engine/Source/Developer/ProjectLauncher/Private/Widgets/Shared/SProjectLauncherValidation.h b/Engine/Source/Developer/ProjectLauncher/Private/Widgets/Shared/SProjectLauncherValidation.h
new file mode 100644
--- /dev/null
+++ b/Engine/Source/Developer/ProjectLauncher/Private/Widgets/Shared/SProjectLauncherValidation.h
@@ -0,0 +1,261 @@
+// Copyright 1998-2017 Epic Games, Inc. All Rights Reserved.
+
+#pragma once
+
+#include "CoreMinimal.h"
+#include "ILauncherProfile.h"
+#include "Misc/Attribute.h"
+#include "Layout/Visibility.h"
+#include "Widgets/DeclarativeSyntaxSupport.h"
+#include "Widgets/SWidget.h"
+#include "Widgets/SCompoundWidget.h"
+#include "Widgets/SBoxPanel.h"
+#include "EditorStyleSet.h"
+#include "Widgets/Text/STextBlock.h"
+#include "Widgets/Images/SImage.h"
+
+
+#define LOCTEXT_NAMESPACE "SProjectLauncherValidation"
+
+
+/**
+ * Implements the launcher's profile validation panel.
+ */
+class SProjectLauncherValidation
+	: public SCompoundWidget
+{
+public:
+
+	SLATE_BEGIN_ARGS(SProjectLauncherValidation) { }
+		SLATE_ATTRIBUTE(TSharedPtr<ILauncherProfile>, LaunchProfile)
+	SLATE_END_ARGS()
+
+
+public:
+
+	/**
+	 * Constructs the widget.
+	 *
+	 * @param InArgs The Slate argument list.
+	 * @param InModel The data model.
+	 */
+	void Construct( const FArguments& InArgs)
+	{
+		LaunchProfileAttr = InArgs._LaunchProfile;
+
+		ChildSlot
+		[
+			SNew(SVerticalBox)
+			
+			// build settings
+			+ SVerticalBox::Slot()
+				.AutoHeight()
+				[
+					MakeValidationMessage(TEXT("Icons.Error"), LOCTEXT("NoBuildGameSelectedError", "A Project must be selected.").ToString(), ELauncherProfileValidationErrors::NoProjectSelected)
+				]
+
+			+ SVerticalBox::Slot()
+				.AutoHeight()
+				[
+					MakeValidationMessage(TEXT("Icons.Error"), LOCTEXT("NoBuildConfigurationSelectedError", "A Build Configuration must be selected.").ToString(), ELauncherProfileValidationErrors::NoBuildConfigurationSelected)
+				]
+
+			// cook settings
+			+ SVerticalBox::Slot()
+				.AutoHeight()
+				[
+					MakeValidationMessage(TEXT("Icons.Error"), LOCTEXT("NoCookedPlatformSelectedError", "At least one Platform must be selected when cooking by the book.").ToString(), ELauncherProfileValidationErrors::NoPlatformSelected)
+				]
+
+			+ SVerticalBox::Slot()
+				.AutoHeight()
+				[
+					MakeValidationMessage(TEXT("Icons.Error"), LOCTEXT("NoCookedCulturesSelectedError", "At least one Culture must be selected when cooking by the book.").ToString(), ELauncherProfileValidationErrors::NoCookedCulturesSelected)
+				]
+
+			// packaging settings
+
+			// archive step
+			+ SVerticalBox::Slot()
+				.AutoHeight()
+				[
+					MakeValidationMessage(TEXT("Icons.Error"), LOCTEXT("NoArchiveDirectorySpecifiedError", "The archive step requires a valid directory.").ToString(), ELauncherProfileValidationErrors::NoArchiveDirectorySpecified)
+				]
+
+			// deployment settings
+			+ SVerticalBox::Slot()
+				.AutoHeight()
+				[
+					MakeValidationMessage(TEXT("Icons.Error"), LOCTEXT("CopyToDeviceRequiresCookByTheBookError", "Deployment by copying to device requires 'By The Book' cooking.").ToString(), ELauncherProfileValidationErrors::CopyToDeviceRequiresCookByTheBook)
+				]
+
+			+ SVerticalBox::Slot()
+				.AutoHeight()
+				[
+					MakeValidationMessage(TEXT("Icons.Error"), LOCTEXT("DeployedDeviceGroupRequired", "A device group must be selected when deploying builds.").ToString(), ELauncherProfileValidationErrors::DeployedDeviceGroupRequired)
+				]
+
+			// launch settings
+			+ SVerticalBox::Slot()
+				.AutoHeight()
+				[
+					MakeValidationMessage(TEXT("Icons.Error"), LOCTEXT("CustomRolesNotSupportedYet", "Custom launch roles are not supported yet.").ToString(), ELauncherProfileValidationErrors::CustomRolesNotSupportedYet)
+				]
+
+			+ SVerticalBox::Slot()
+				.AutoHeight()
+				[
+					MakeValidationMessage(TEXT("Icons.Error"), LOCTEXT("InitialCultureNotAvailable", "The Initial Culture selected for launch is not in the build.").ToString(), ELauncherProfileValidationErrors::InitialCultureNotAvailable)
+				]
+
+			+ SVerticalBox::Slot()
+				.AutoHeight()
+				[
+					MakeValidationMessage(TEXT("Icons.Error"), LOCTEXT("InitialMapNotAvailable", "The Initial Map selected for launch is not in the build.").ToString(), ELauncherProfileValidationErrors::InitialMapNotAvailable)
+				]
+
+			+ SVerticalBox::Slot()
+				.AutoHeight()
+				[
+					MakeValidationMessage(TEXT("Icons.Error"), LOCTEXT("NoLaunchRoleDeviceAssigned", "One or more launch roles do not have a device assigned.").ToString(), ELauncherProfileValidationErrors::NoLaunchRoleDeviceAssigned)
+				]
+
+			+ SVerticalBox::Slot()
+				.AutoHeight()
+				[
+					MakeValidationMessage(TEXT("Icons.Error"), LOCTEXT("GeneratingChunksRequiresUnrealPak", "UnrealPak must be selected to Generate Chunks.").ToString(), ELauncherProfileValidationErrors::GeneratingChunksRequiresUnrealPak)
+				]
+			
+			+ SVerticalBox::Slot()
+				.AutoHeight()
+				[
+					MakeValidationMessage(TEXT("Icons.Error"), LOCTEXT("GeneratingHttpChunkDataRequiresGeneratingChunks", "Generate Chunks must be selected to Generate Http Chunk Install Data.").ToString(), ELauncherProfileValidationErrors::GeneratingHttpChunkDataRequiresGeneratingChunks)
+				]
+
+			+ SVerticalBox::Slot()
+				.AutoHeight()
+				[
+					MakeValidationMessage(TEXT("Icons.Error"), LOCTEXT("GeneratingHttpChunkDataRequiresValidDirectoryAndName", "Generating Http Chunk Install Data requires a valid directory and release name.").ToString(), ELauncherProfileValidationErrors::GeneratingHttpChunkDataRequiresValidDirectoryAndName)
+				]
+
+			+ SVerticalBox::Slot()
+				.AutoHeight()
+				[
+					MakeValidationMessage(TEXT("Icons.Error"), LOCTEXT("ShippingDoesntSupportCommandlineOptionsCantUseCookOnTheFly", "Shipping doesn't support commandline options and can't use cook on the fly").ToString(), ELauncherProfileValidationErrors::ShippingDoesntSupportCommandlineOptionsCantUseCookOnTheFly)
+				]
+
+			+ SVerticalBox::Slot()
+				.AutoHeight()
+				[
+					MakeValidationMessage(TEXT("Icons.Error"), LOCTEXT("CookOnTheFlyDoesntSupportServer", "Cook on the fly doesn't support server target configurations").ToString(), ELauncherProfileValidationErrors::CookOnTheFlyDoesntSupportServer)
+				]
+
+			+ SVerticalBox::Slot()
+				.AutoHeight()
+				[
+					MakeValidationMessage(TEXT("Icons.Error"), LOCTEXT("UnversionedAndIncrimental", "Unversioned build cannot be incremental.").ToString(), ELauncherProfileValidationErrors::UnversionedAndIncrimental)
+				]
+
+			+ SVerticalBox::Slot()
+				.AutoHeight()
+				[
+					MakeCallbackMessage(TEXT("Icons.Error"), ELauncherProfileValidationErrors::NoPlatformSDKInstalled)
+				]
+		];
+	}
+
+protected:
+
+	/**
+	 * Creates a widget for a validation message.
+	 *
+	 * @param IconName The name of the message icon.
+	 * @param MessageText The message text.
+	 * @param MessageType The message type.
+	 */
+	TSharedRef<SWidget> MakeValidationMessage( const TCHAR* IconName, const FString& MessageText, ELauncherProfileValidationErrors::Type Message )
+	{
+		return SNew(SHorizontalBox)
+			.Visibility(this, &SProjectLauncherValidation::HandleValidationMessageVisibility, Message)
+
+		+ SHorizontalBox::Slot()
+			.AutoWidth()
+			.Padding(2.0)
+			[
+				SNew(SImage)
+					.Image(FEditorStyle::GetBrush(IconName))
+			]
+
+		+ SHorizontalBox::Slot()
+			.AutoWidth()
+			.VAlign(VAlign_Center)
+			[
+				SNew(STextBlock)
+					.Text(FText::FromString(MessageText))
+			];
+	}
+
+	/**
+	 * Creates a widget for a validation message.
+	 *
+	 * @param IconName The name of the message icon.
+	 * @param MessageType The message type.
+	 */
+	TSharedRef<SWidget> MakeCallbackMessage( const TCHAR* IconName,ELauncherProfileValidationErrors::Type Message )
+	{
+		return SNew(SHorizontalBox)
+			.Visibility(this, &SProjectLauncherValidation::HandleValidationMessageVisibility, Message)
+
+		+ SHorizontalBox::Slot()
+			.AutoWidth()
+			.Padding(2.0)
+			[
+				SNew(SImage)
+					.Image(FEditorStyle::GetBrush(IconName))
+			]
+
+		+ SHorizontalBox::Slot()
+			.AutoWidth()
+			.VAlign(VAlign_Center)
+			[
+				SNew(STextBlock)
+					.Text(this, &SProjectLauncherValidation::HandleValidationMessage, Message)
+			];
+	}
+
+private:
+
+	// Callback for getting the visibility state of a validation message.
+	EVisibility HandleValidationMessageVisibility( ELauncherProfileValidationErrors::Type Error ) const
+	{
+		TSharedPtr<ILauncherProfile> LaunchProfile = LaunchProfileAttr.Get();
+		if (!LaunchProfile.IsValid() || LaunchProfile->HasValidationError(Error))
+		{
+			return EVisibility::Visible;
+		}
+		return EVisibility::Collapsed;
+	}
+
+	FText HandleValidationMessage( ELauncherProfileValidationErrors::Type Error ) const
+	{
+		TSharedPtr<ILauncherProfile> LaunchProfile = LaunchProfileAttr.Get();
+		if (LaunchProfile.IsValid())
+		{
+			if (LaunchProfile->HasValidationError(Error))
+			{
+				return FText::Format(LOCTEXT("NoPlatformSDKInstalledFmt", "A required platform SDK is mising: {0}"), FText::FromString(LaunchProfile->GetInvalidPlatform()));
+			}
+
+			return FText::GetEmpty();
+		}
+		return LOCTEXT("InvalidLaunchProfile", "Invalid Launch Profile.");
+	}
+
+private:
+
+	// Attribute for the launch profile this widget shows validation for. 
+	TAttribute<TSharedPtr<ILauncherProfile>> LaunchProfileAttr;
+};
+
+
+#undef LOCTEXT_NAMESPACE
diff --git a/Engine/Source/Developer/ProjectLauncher/Private/Widgets/Shared/SProjectLauncherVariantSelector.h b/Engine/Source/Developer/ProjectLauncher/Private/Widgets/Shared/SProjectLauncherVariantSelector.h
--- a/Engine/Source/Developer/ProjectLauncher/Private/Widgets/Shared/SProjectLauncherVariantSelector.h
+++ b/Engine/Source/Developer/ProjectLauncher/Private/Widgets/Shared/SProjectLauncherVariantSelector.h
@@ -6,7 +6,7 @@
 #include "Widgets/DeclarativeSyntaxSupport.h"
 #include "Layout/Margin.h"
 #include "Widgets/SCompoundWidget.h"
-#include "Interfaces/ITargetDeviceProxy.h"
+#include "ITargetDeviceProxy.h"
 #include "Styling/CoreStyle.h"
 #include "Widgets/Text/STextBlock.h"
 #include "Widgets/Input/SComboButton.h"
@@ -46,7 +46,7 @@ public:
 	 * @param InArgs The Slate argument list.
 	 * @param InModel The data model.
 	 */
-	void Construct(const FArguments& InArgs, ITargetDeviceProxyPtr DeviceProxy)
+	void Construct(const FArguments& InArgs, TSharedPtr<ITargetDeviceProxy> DeviceProxy)
 	{
 		OnVariantSelected = InArgs._OnVariantSelected;
 
diff --git a/Engine/Source/Developer/ProjectLauncher/ProjectLauncher.Build.cs b/Engine/Source/Developer/ProjectLauncher/ProjectLauncher.Build.cs
--- a/Engine/Source/Developer/ProjectLauncher/ProjectLauncher.Build.cs
+++ b/Engine/Source/Developer/ProjectLauncher/ProjectLauncher.Build.cs
@@ -43,16 +43,18 @@ public class ProjectLauncher : ModuleRules
 				"Developer/ProjectLauncher/Private",
 				"Developer/ProjectLauncher/Private/Models",
 				"Developer/ProjectLauncher/Private/Widgets",
-				"Developer/ProjectLauncher/Private/Widgets/Build",
+                "Developer/ProjectLauncher/Private/Widgets/Archive",
+                "Developer/ProjectLauncher/Private/Widgets/Build",
 				"Developer/ProjectLauncher/Private/Widgets/Cook",
 				"Developer/ProjectLauncher/Private/Widgets/Deploy",
 				"Developer/ProjectLauncher/Private/Widgets/Launch",
 				"Developer/ProjectLauncher/Private/Widgets/Package",
                 "Developer/ProjectLauncher/Private/Widgets/Profile",
-				"Developer/ProjectLauncher/Private/Widgets/Project",
-				"Developer/ProjectLauncher/Private/Widgets/Shared",
+                "Developer/ProjectLauncher/Private/Widgets/Progress",
+                "Developer/ProjectLauncher/Private/Widgets/Project",
+                "Developer/ProjectLauncher/Private/Widgets/Settings",
+                "Developer/ProjectLauncher/Private/Widgets/Shared",
 				"Developer/ProjectLauncher/Private/Widgets/Toolbar",
-                "Developer/ProjectLauncher/Private/Widgets/Archive",
             }
 		);
 	}
diff --git a/Engine/Source/Developer/ProjectLauncher/Public/IProjectLauncherModule.h b/Engine/Source/Developer/ProjectLauncher/Public/IProjectLauncherModule.h
new file mode 100644
--- /dev/null
+++ b/Engine/Source/Developer/ProjectLauncher/Public/IProjectLauncherModule.h
@@ -0,0 +1,18 @@
+// Copyright 1998-2017 Epic Games, Inc. All Rights Reserved.
+
+#pragma once
+
+#include "Modules/ModuleInterface.h"
+
+
+/**
+ * Interface for launcher UI modules.
+ */
+class IProjectLauncherModule
+	: public IModuleInterface
+{
+public:
+
+	/** Virtual destructor. */
+	virtual ~IProjectLauncherModule() { }
+};
diff --git a/Engine/Source/Developer/ScreenShotComparison/Private/Widgets/SAsyncImage.cpp b/Engine/Source/Developer/ScreenShotComparison/Private/Widgets/SAsyncImage.cpp
--- a/Engine/Source/Developer/ScreenShotComparison/Private/Widgets/SAsyncImage.cpp
+++ b/Engine/Source/Developer/ScreenShotComparison/Private/Widgets/SAsyncImage.cpp
@@ -1,12 +1,15 @@
 // Copyright 1998-2017 Epic Games, Inc. All Rights Reserved.
 
 #include "Widgets/SAsyncImage.h"
-#include "Widgets/SOverlay.h"
+
 #include "HAL/FileManager.h"
+#include "IImageWrapper.h"
+#include "IImageWrapperModule.h"
 #include "Misc/FileHelper.h"
-#include "Interfaces/IImageWrapperModule.h"
-#include "SlateApplication.h"
 #include "ModuleManager.h"
+#include "SlateApplication.h"
+#include "Widgets/SOverlay.h"
+
 
 void SAsyncImage::Construct(const FArguments& InArgs)
 {
@@ -86,14 +89,14 @@ FSlateTextureDataPtr SAsyncImage::LoadScreenshot(FString ImagePath)
 	if ( FFileHelper::LoadFileToArray(RawFileData, *ImagePath) )
 	{
 		IImageWrapperModule& ImageWrapperModule = FModuleManager::GetModuleChecked<IImageWrapperModule>(FName("ImageWrapper"));
-		IImageWrapperPtr ImageWrappers[3] =
+		TSharedPtr<IImageWrapper> ImageWrappers[3] =
 		{
 			ImageWrapperModule.CreateImageWrapper(EImageFormat::PNG),
 			ImageWrapperModule.CreateImageWrapper(EImageFormat::JPEG),
 			ImageWrapperModule.CreateImageWrapper(EImageFormat::BMP),
 		};
 
-		for ( IImageWrapperPtr ImageWrapper : ImageWrappers )
+		for ( TSharedPtr<IImageWrapper> ImageWrapper : ImageWrappers )
 		{
 			if ( ImageWrapper.IsValid() && ImageWrapper->SetCompressed(RawFileData.GetData(), RawFileData.Num()) )
 			{
diff --git a/Engine/Source/Developer/ScreenShotComparison/Private/Widgets/SScreenComparisonRow.cpp b/Engine/Source/Developer/ScreenShotComparison/Private/Widgets/SScreenComparisonRow.cpp
--- a/Engine/Source/Developer/ScreenShotComparison/Private/Widgets/SScreenComparisonRow.cpp
+++ b/Engine/Source/Developer/ScreenShotComparison/Private/Widgets/SScreenComparisonRow.cpp
@@ -16,7 +16,7 @@
 #include "Widgets/Input/SSlider.h"
 #include "HAL/FileManager.h"
 #include "Misc/FileHelper.h"
-#include "Interfaces/IImageWrapperModule.h"
+#include "IImageWrapperModule.h"
 #include "Framework/Application/SlateApplication.h"
 #include "JsonObjectConverter.h"
 #include "SScreenShotImagePopup.h"
diff --git a/Engine/Source/Developer/ScreenShotComparisonTools/Private/ImageComparer.cpp b/Engine/Source/Developer/ScreenShotComparisonTools/Private/ImageComparer.cpp
--- a/Engine/Source/Developer/ScreenShotComparisonTools/Private/ImageComparer.cpp
+++ b/Engine/Source/Developer/ScreenShotComparisonTools/Private/ImageComparer.cpp
@@ -1,12 +1,14 @@
 // Copyright 1998-2017 Epic Games, Inc. All Rights Reserved.
 
 #include "ImageComparer.h"
+
+#include "Async/ParallelFor.h"
+#include "IImageWrapper.h"
+#include "IImageWrapperModule.h"
 #include "Misc/FileHelper.h"
 #include "Misc/Paths.h"
 #include "Modules/ModuleManager.h"
 
-#include "Interfaces/IImageWrapperModule.h"
-#include "Async/ParallelFor.h"
 
 #define LOCTEXT_NAMESPACE "ImageComparer"
 
@@ -15,6 +17,7 @@ const FImageTolerance FImageTolerance::DefaultIgnoreLess(16, 16, 16, 16, 16, 240
 const FImageTolerance FImageTolerance::DefaultIgnoreAntiAliasing(32, 32, 32, 32, 64, 96, true, false, 0.02f, 0.02f);
 const FImageTolerance FImageTolerance::DefaultIgnoreColors(16, 16, 16, 16, 16, 240, false, true, 0.02f, 0.02f);
 
+
 class FImageDelta
 {
 public:
@@ -83,7 +86,7 @@ public:
 		FString TempDeltaFile = FPaths::CreateTempFilename(*TempDir, TEXT("ImageCompare-"), TEXT(".png"));
 
 		IImageWrapperModule& ImageWrapperModule = FModuleManager::GetModuleChecked<IImageWrapperModule>(FName("ImageWrapper"));
-		IImageWrapperPtr ImageWriter = ImageWrapperModule.CreateImageWrapper(EImageFormat::PNG);
+		TSharedPtr<IImageWrapper> ImageWriter = ImageWrapperModule.CreateImageWrapper(EImageFormat::PNG);
 
 		if ( ImageWriter.IsValid() )
 		{
@@ -232,7 +235,7 @@ void FComparableImage::Process()
 TSharedPtr<FComparableImage> FImageComparer::Open(const FString& ImagePath, FText& OutError)
 {
 	IImageWrapperModule& ImageWrapperModule = FModuleManager::GetModuleChecked<IImageWrapperModule>(FName("ImageWrapper"));
-	IImageWrapperPtr ImageReader = ImageWrapperModule.CreateImageWrapper(EImageFormat::PNG);
+	TSharedPtr<IImageWrapper> ImageReader = ImageWrapperModule.CreateImageWrapper(EImageFormat::PNG);
 
 	if ( !ImageReader.IsValid() )
 	{
diff --git a/Engine/Source/Developer/ScreenShotComparisonTools/Private/ScreenShotManager.cpp b/Engine/Source/Developer/ScreenShotComparisonTools/Private/ScreenShotManager.cpp
--- a/Engine/Source/Developer/ScreenShotComparisonTools/Private/ScreenShotManager.cpp
+++ b/Engine/Source/Developer/ScreenShotComparisonTools/Private/ScreenShotManager.cpp
@@ -1,20 +1,21 @@
 // Copyright 1998-2017 Epic Games, Inc. All Rights Reserved.
 
 #include "ScreenShotManager.h"
-#include "HAL/FileManager.h"
+
+#include "AutomationWorkerMessages.h"
 #include "Async/Async.h"
+#include "HAL/FileManager.h"
+#include "MessageEndpointBuilder.h"
 #include "Misc/FileHelper.h"
 #include "Misc/Paths.h"
 #include "Misc/EngineVersion.h"
 #include "Misc/FilterCollection.h"
-#include "AutomationWorkerMessages.h"
-#include "Helpers/MessageEndpointBuilder.h"
 #include "Modules/ModuleManager.h"
-
 #include "Serialization/JsonReader.h"
 #include "Serialization/JsonSerializer.h"
 #include "JsonObjectConverter.h"
 
+
 class FScreenshotComparisons
 {
 public:
diff --git a/Engine/Source/Developer/ScreenShotComparisonTools/Private/ScreenShotManager.h b/Engine/Source/Developer/ScreenShotComparisonTools/Private/ScreenShotManager.h
--- a/Engine/Source/Developer/ScreenShotComparisonTools/Private/ScreenShotManager.h
+++ b/Engine/Source/Developer/ScreenShotComparisonTools/Private/ScreenShotManager.h
@@ -10,11 +10,12 @@
 #include "Async/Future.h"
 #include "ImageComparer.h"
 #include "Interfaces/IScreenShotManager.h"
-#include "Helpers/MessageEndpoint.h"
+#include "MessageEndpoint.h"
 
 class FScreenshotComparisons;
 struct FScreenShotDataItem;
 
+
 /**
  * Implements the ScreenShotManager that contains screen shot data.
  */
diff --git a/Engine/Source/Developer/ScreenShotComparisonTools/Private/ScreenShotToolsModule.cpp b/Engine/Source/Developer/ScreenShotComparisonTools/Private/ScreenShotToolsModule.cpp
--- a/Engine/Source/Developer/ScreenShotComparisonTools/Private/ScreenShotToolsModule.cpp
+++ b/Engine/Source/Developer/ScreenShotComparisonTools/Private/ScreenShotToolsModule.cpp
@@ -7,11 +7,9 @@
 #include "CoreMinimal.h"
 #include "Modules/ModuleManager.h"
 #include "Interfaces/IScreenShotToolsModule.h"
-#include "IMessagingModule.h"
 #include "ScreenShotManager.h"
 
 
-
 /**
  * Implements the ScreenshotTools module.
  */
diff --git a/Engine/Source/Developer/ScreenShotComparisonTools/ScreenShotComparisonTools.Build.cs b/Engine/Source/Developer/ScreenShotComparisonTools/ScreenShotComparisonTools.Build.cs
--- a/Engine/Source/Developer/ScreenShotComparisonTools/ScreenShotComparisonTools.Build.cs
+++ b/Engine/Source/Developer/ScreenShotComparisonTools/ScreenShotComparisonTools.Build.cs
@@ -16,19 +16,19 @@ public class ScreenShotComparisonTools : ModuleRules
 		PrivateDependencyModuleNames.AddRange(
 			new string[] {
 				"AutomationMessages",
-				"UnrealEdMessages",
-				"Slate",
                 "EditorStyle",
 				"ImageWrapper",
 				"Json",
-				"JsonUtilities"
-			}
-		);
+				"JsonUtilities",
+				"Slate",
+                "UnrealEdMessages",
+            }
+        );
 
 		PrivateIncludePathModuleNames.AddRange(
 			new string[] {
-				"Messaging"
-			}
+                "MessagingCommon",
+            }
 		);
 
 		PrivateIncludePaths.AddRange(
diff --git a/Engine/Source/Developer/SessionFrontend/Private/Widgets/SSessionFrontend.cpp b/Engine/Source/Developer/SessionFrontend/Private/Widgets/SSessionFrontend.cpp
--- a/Engine/Source/Developer/SessionFrontend/Private/Widgets/SSessionFrontend.cpp
+++ b/Engine/Source/Developer/SessionFrontend/Private/Widgets/SSessionFrontend.cpp
@@ -1,6 +1,8 @@
 // Copyright 1998-2017 Epic Games, Inc. All Rights Reserved.
 
 #include "Widgets/SSessionFrontend.h"
+
+#include "ITargetDeviceProxyManager.h"
 #include "Modules/ModuleManager.h"
 #include "Widgets/SBoxPanel.h"
 #include "Textures/SlateIcon.h"
@@ -8,9 +10,9 @@
 #include "Framework/Docking/WorkspaceItem.h"
 #include "Framework/Docking/TabManager.h"
 #include "EditorStyleSet.h"
-#include "Interfaces/ITargetDeviceServicesModule.h"
-#include "Interfaces/IAutomationControllerModule.h"
-#include "Interfaces/IAutomationWindowModule.h"
+#include "ITargetDeviceServicesModule.h"
+#include "IAutomationControllerModule.h"
+#include "IAutomationWindowModule.h"
 #include "Interfaces/IScreenShotToolsModule.h"
 #include "Interfaces/IScreenShotComparisonModule.h"
 #include "ISessionServicesModule.h"
diff --git a/Engine/Source/Developer/SessionFrontend/Private/Widgets/SSessionFrontend.h b/Engine/Source/Developer/SessionFrontend/Private/Widgets/SSessionFrontend.h
--- a/Engine/Source/Developer/SessionFrontend/Private/Widgets/SSessionFrontend.h
+++ b/Engine/Source/Developer/SessionFrontend/Private/Widgets/SSessionFrontend.h
@@ -7,12 +7,12 @@
 #include "Widgets/DeclarativeSyntaxSupport.h"
 #include "Widgets/SCompoundWidget.h"
 #include "ISessionManager.h"
-#include "Interfaces/ITargetDeviceProxyManager.h"
 #include "Interfaces/IScreenShotManager.h"
 
 class FMenuBuilder;
 class FSpawnTabArgs;
 class FTabManager;
+class ITargetDeviceProxyManager;
 class SButton;
 class SWindow;
 
@@ -64,7 +64,7 @@ private:
 private:
 
 	/** Holds the target device proxy manager. */
-	ITargetDeviceProxyManagerPtr DeviceProxyManager;
+	TSharedPtr<ITargetDeviceProxyManager> DeviceProxyManager;
 	
 	/** Holds a flag indicating whether the launcher overlay is visible. */
 	bool LauncherOverlayVisible;
diff --git a/Engine/Source/Developer/SimplygonSwarm/Private/SimplygonSwarm.cpp b/Engine/Source/Developer/SimplygonSwarm/Private/SimplygonSwarm.cpp
--- a/Engine/Source/Developer/SimplygonSwarm/Private/SimplygonSwarm.cpp
+++ b/Engine/Source/Developer/SimplygonSwarm/Private/SimplygonSwarm.cpp
@@ -7,7 +7,8 @@
 #include "UObject/Object.h"
 #include "Misc/Paths.h"
 #include "Misc/ScopedSlowTask.h"
-#include "Interfaces/IImageWrapperModule.h"
+#include "IImageWrapper.h"
+#include "IImageWrapperModule.h"
 #include "Editor/EditorPerProjectUserSettings.h"
 #include "Misc/EngineVersion.h"
 #include "Misc/MonitoredProcess.h"
@@ -1404,7 +1405,7 @@ bool ZipContentsForUpload(FString InputDirectoryPath, FString OutputFileName)
 		)
 	{
 		IImageWrapperModule& ImageWrapperModule = FModuleManager::LoadModuleChecked<IImageWrapperModule>(FName("ImageWrapper"));
-		IImageWrapperPtr ImageWrapper = ImageWrapperModule.CreateImageWrapper(EImageFormat::PNG);
+		TSharedPtr<IImageWrapper> ImageWrapper = ImageWrapperModule.CreateImageWrapper(EImageFormat::PNG);
 		 
 		TArray<uint8> TextureData;
 		if (!FFileHelper::LoadFileToArray(TextureData, *FPaths::ConvertRelativePathToFull(TexturePath)) && TextureData.Num() > 0)
@@ -1465,7 +1466,7 @@ bool ZipContentsForUpload(FString InputDirectoryPath, FString OutputFileName)
 		{
 
 			IImageWrapperModule& ImageWrapperModule = FModuleManager::LoadModuleChecked<IImageWrapperModule>(FName("ImageWrapper"));
-			IImageWrapperPtr ImageWrapper = ImageWrapperModule.CreateImageWrapper(EImageFormat::PNG);
+			TSharedPtr<IImageWrapper> ImageWrapper = ImageWrapperModule.CreateImageWrapper(EImageFormat::PNG);
 			
 			FString TextureOutputRelative = FString::Printf(TEXT("%s/%s.png"), ANSI_TO_TCHAR(SsfTextureTable->TexturesDirectory->Value.c_str()), *TextureName);
 			FString TextureOutputPath = FString::Printf(TEXT("%s%s"), *BaseTexturePath, *TextureOutputRelative);
diff --git a/Engine/Source/Developer/SimplygonSwarm/Public/SimplygonSwarmCommon.h b/Engine/Source/Developer/SimplygonSwarm/Public/SimplygonSwarmCommon.h
--- a/Engine/Source/Developer/SimplygonSwarm/Public/SimplygonSwarmCommon.h
+++ b/Engine/Source/Developer/SimplygonSwarm/Public/SimplygonSwarmCommon.h
@@ -11,7 +11,7 @@
 THIRD_PARTY_INCLUDES_START
 #include "ssf.h"
 THIRD_PARTY_INCLUDES_END
-#include "Interfaces/IImageWrapper.h"
+#include "IImageWrapper.h"
 #include "SPLInclude.h"
 
 namespace SPL = Simplygon::SPL::v80;
diff --git a/Engine/Source/Developer/SlateReflector/Private/WidgetSnapshotService.cpp b/Engine/Source/Developer/SlateReflector/Private/WidgetSnapshotService.cpp
--- a/Engine/Source/Developer/SlateReflector/Private/WidgetSnapshotService.cpp
+++ b/Engine/Source/Developer/SlateReflector/Private/WidgetSnapshotService.cpp
@@ -4,7 +4,7 @@
 #include "Misc/App.h"
 #include "WidgetSnapshotMessages.h"
 #include "IMessagingModule.h"
-#include "Helpers/MessageEndpointBuilder.h"
+#include "MessageEndpointBuilder.h"
 #include "Widgets/SWidgetSnapshotVisualizer.h"
 
 FWidgetSnapshotService::FWidgetSnapshotService()
@@ -48,7 +48,7 @@ void FWidgetSnapshotService::AbortSnapshotRequest(const FGuid& InSnapshotRequest
 	PendingSnapshotResponseHandlers.Remove(InSnapshotRequestId);
 }
 
-void FWidgetSnapshotService::HandleWidgetSnapshotRequestMessage(const FWidgetSnapshotRequest& Message, const IMessageContextRef& Context)
+void FWidgetSnapshotService::HandleWidgetSnapshotRequestMessage(const FWidgetSnapshotRequest& Message, const TSharedRef<IMessageContext, ESPMode::ThreadSafe>& Context)
 {
 	if (MessageEndpoint.IsValid() && Message.TargetInstanceId == FApp::GetInstanceId())
 	{
@@ -63,7 +63,7 @@ void FWidgetSnapshotService::HandleWidgetSnapshotRequestMessage(const FWidgetSna
 	}
 }
 
-void FWidgetSnapshotService::HandleWidgetSnapshotResponseMessage(const FWidgetSnapshotResponse& Message, const IMessageContextRef& Context)
+void FWidgetSnapshotService::HandleWidgetSnapshotResponseMessage(const FWidgetSnapshotResponse& Message, const TSharedRef<IMessageContext, ESPMode::ThreadSafe>& Context)
 {
 	const FOnWidgetSnapshotResponse* FoundResponseHandler = PendingSnapshotResponseHandlers.Find(Message.SnapshotRequestId);
 	if (FoundResponseHandler && FoundResponseHandler->IsBound())
diff --git a/Engine/Source/Developer/SlateReflector/Private/WidgetSnapshotService.h b/Engine/Source/Developer/SlateReflector/Private/WidgetSnapshotService.h
--- a/Engine/Source/Developer/SlateReflector/Private/WidgetSnapshotService.h
+++ b/Engine/Source/Developer/SlateReflector/Private/WidgetSnapshotService.h
@@ -4,7 +4,7 @@
 
 #include "CoreMinimal.h"
 #include "Misc/Guid.h"
-#include "Helpers/MessageEndpoint.h"
+#include "MessageEndpoint.h"
 
 struct FWidgetSnapshotRequest;
 struct FWidgetSnapshotResponse;
@@ -27,13 +27,13 @@ public:
 
 private:
 	/** Handles FWidgetSnapshotRequest messages */
-	void HandleWidgetSnapshotRequestMessage(const FWidgetSnapshotRequest& Message, const IMessageContextRef& Context);
+	void HandleWidgetSnapshotRequestMessage(const FWidgetSnapshotRequest& Message, const TSharedRef<IMessageContext, ESPMode::ThreadSafe>& Context);
 
 	/** Handles FWidgetSnapshotResponse messages */
-	void HandleWidgetSnapshotResponseMessage(const FWidgetSnapshotResponse& Message, const IMessageContextRef& Context);
+	void HandleWidgetSnapshotResponseMessage(const FWidgetSnapshotResponse& Message, const TSharedRef<IMessageContext, ESPMode::ThreadSafe>& Context);
 
 	/** Holds the message endpoint */
-	FMessageEndpointPtr MessageEndpoint;
+	TSharedPtr<FMessageEndpoint, ESPMode::ThreadSafe> MessageEndpoint;
 
 	/** Handlers awaiting their resultant snapshot data */
 	TMap<FGuid, FOnWidgetSnapshotResponse> PendingSnapshotResponseHandlers;
diff --git a/Engine/Source/Developer/SlateReflector/SlateReflector.Build.cs b/Engine/Source/Developer/SlateReflector/SlateReflector.Build.cs
--- a/Engine/Source/Developer/SlateReflector/SlateReflector.Build.cs
+++ b/Engine/Source/Developer/SlateReflector/SlateReflector.Build.cs
@@ -26,7 +26,8 @@ public class SlateReflector : ModuleRules
 
 		PrivateIncludePathModuleNames.AddRange(
 			new string[] {
-				"Messaging",
+                "Messaging",
+				"MessagingCommon",
 			}
 		);
 
diff --git a/Engine/Source/Developer/StandaloneRenderer/Private/OpenGL/SlateOpenGLTextureManager.cpp b/Engine/Source/Developer/StandaloneRenderer/Private/OpenGL/SlateOpenGLTextureManager.cpp
--- a/Engine/Source/Developer/StandaloneRenderer/Private/OpenGL/SlateOpenGLTextureManager.cpp
+++ b/Engine/Source/Developer/StandaloneRenderer/Private/OpenGL/SlateOpenGLTextureManager.cpp
@@ -7,12 +7,13 @@
 #include "Styling/SlateStyleRegistry.h"
 #include "StandaloneRendererPlatformHeaders.h"
 #include "OpenGL/SlateOpenGLTextures.h"
+#include "IImageWrapper.h"
+#include "IImageWrapperModule.h"
 
 
-#include "Interfaces/IImageWrapperModule.h"
-
 DEFINE_LOG_CATEGORY_STATIC(LogSlateOpenGL, Log, All);
 
+
 FSlateOpenGLTextureManager::FDynamicTextureResource::FDynamicTextureResource( FSlateOpenGLTexture* InOpenGLTexture )
 	: Proxy( new FSlateShaderResourceProxy )
 	, OpenGLTexture( InOpenGLTexture  )
@@ -124,7 +125,7 @@ bool FSlateOpenGLTextureManager::LoadTexture( const FSlateBrush& InBrush, uint32
 	{
 		// We assume all resources are png for now.
 		IImageWrapperModule& ImageWrapperModule = FModuleManager::LoadModuleChecked<IImageWrapperModule>( FName("ImageWrapper") );
-		IImageWrapperPtr ImageWrapper = ImageWrapperModule.CreateImageWrapper( EImageFormat::PNG );
+		TSharedPtr<IImageWrapper> ImageWrapper = ImageWrapperModule.CreateImageWrapper( EImageFormat::PNG );
 		if ( ImageWrapper.IsValid() && ImageWrapper->SetCompressed( Buffer.GetData(), Buffer.Num() ) )
 		{
 			// Determine the block size.  This is bytes per pixel
diff --git a/Engine/Source/Developer/StandaloneRenderer/Private/Windows/D3D/SlateD3DTextureManager.cpp b/Engine/Source/Developer/StandaloneRenderer/Private/Windows/D3D/SlateD3DTextureManager.cpp
--- a/Engine/Source/Developer/StandaloneRenderer/Private/Windows/D3D/SlateD3DTextureManager.cpp
+++ b/Engine/Source/Developer/StandaloneRenderer/Private/Windows/D3D/SlateD3DTextureManager.cpp
@@ -6,13 +6,15 @@
 #include "StandaloneRendererPrivate.h"
 #include "Misc/FileHelper.h"
 #include "Modules/ModuleManager.h"
-#include "Interfaces/IImageWrapperModule.h"
-
+#include "IImageWrapper.h"
+#include "IImageWrapperModule.h"
 #include "Styling/SlateStyle.h"
 #include "Styling/SlateStyleRegistry.h"
 
+
 DEFINE_LOG_CATEGORY_STATIC(LogSlateD3D, Log, All);
 
+
 FSlateD3DTextureManager::FDynamicTextureResource::FDynamicTextureResource( FSlateD3DTexture* ExistingTexture )
 	: Proxy( new FSlateShaderResourceProxy )
 	, D3DTexture( ExistingTexture  )
@@ -212,7 +214,7 @@ bool FSlateD3DTextureManager::LoadTexture( const FSlateBrush& InBrush, uint32& O
 	if( FFileHelper::LoadFileToArray( RawFileData, *ResourcePath ) )
 	{
 		IImageWrapperModule& ImageWrapperModule = FModuleManager::LoadModuleChecked<IImageWrapperModule>( FName("ImageWrapper") );
-		IImageWrapperPtr ImageWrapper = ImageWrapperModule.CreateImageWrapper( EImageFormat::PNG );
+		TSharedPtr<IImageWrapper> ImageWrapper = ImageWrapperModule.CreateImageWrapper( EImageFormat::PNG );
 		if ( ImageWrapper.IsValid() && ImageWrapper->SetCompressed( RawFileData.GetData(), RawFileData.Num() ) )
 		{
 			OutWidth = ImageWrapper->GetWidth();
diff --git a/Engine/Source/Developer/TargetDeviceServices/Private/Proxies/TargetDeviceProxy.cpp b/Engine/Source/Developer/TargetDeviceServices/Private/Proxies/TargetDeviceProxy.cpp
new file mode 100644
--- /dev/null
+++ b/Engine/Source/Developer/TargetDeviceServices/Private/Proxies/TargetDeviceProxy.cpp
@@ -0,0 +1,301 @@
+// Copyright 1998-2017 Epic Games, Inc. All Rights Reserved.
+
+#include "TargetDeviceProxy.h"
+
+#include "FileMessageAttachment.h"
+#include "HAL/PlatformProcess.h"
+#include "MessageEndpointBuilder.h"
+
+#include "TargetDeviceServiceMessages.h"
+
+
+/* FTargetDeviceProxy structors
+ *****************************************************************************/
+
+FTargetDeviceProxy::FTargetDeviceProxy(const FString& InName)
+	: Connected(false)
+	, LastUpdateTime(0)
+	, Name(InName)
+	, SupportsMultiLaunch(false)
+	, SupportsPowerOff(false)
+	, SupportsPowerOn(false)
+	, SupportsReboot(false)
+{
+
+	InitializeMessaging();
+}
+
+
+FTargetDeviceProxy::FTargetDeviceProxy(const FString& InName, const FTargetDeviceServicePong& Message, const TSharedRef<IMessageContext, ESPMode::ThreadSafe>& Context)
+	: Connected(false)
+	, Name(InName)
+	, SupportsMultiLaunch(false)
+	, SupportsPowerOff(false)
+	, SupportsPowerOn(false)
+	, SupportsReboot(false)
+{
+	UpdateFromMessage(Message, Context);
+	InitializeMessaging();
+}
+
+
+/* FTargetDeviceProxy interface
+*****************************************************************************/
+
+void FTargetDeviceProxy::UpdateFromMessage( const FTargetDeviceServicePong& Message, const TSharedRef<IMessageContext, ESPMode::ThreadSafe>& Context )
+{
+	if (Message.Name != Name)
+	{
+		return;
+	}
+
+	MessageAddress = Context->GetSender();
+
+	Connected = Message.Connected;
+	HostName = Message.HostName;
+	HostUser = Message.HostUser;
+	Make = Message.Make;
+	Model = Message.Model;
+	Name = Message.Name;
+	DeviceUser = Message.DeviceUser;
+	DeviceUserPassword = Message.DeviceUserPassword;
+	Shared = Message.Shared;
+	SupportsMultiLaunch = Message.SupportsMultiLaunch;
+	SupportsPowerOff = Message.SupportsPowerOff;
+	SupportsPowerOn = Message.SupportsPowerOn;
+	SupportsReboot = Message.SupportsReboot;
+	SupportsVariants = Message.SupportsVariants;
+	DefaultVariant = Message.DefaultVariant;
+
+	// Update the map of flavors.
+	for (int Index = 0; Index < Message.Variants.Num(); Index++)
+	{
+		const FTargetDeviceVariant& MsgVariant = Message.Variants[Index];
+
+		FTargetDeviceProxyVariant & Variant = TargetDeviceVariants.Add(MsgVariant.VariantName);
+		Variant.DeviceID = MsgVariant.DeviceID;
+		Variant.VariantName = MsgVariant.VariantName;
+		Variant.TargetPlatformName = MsgVariant.TargetPlatformName;
+		Variant.TargetPlatformId = MsgVariant.TargetPlatformId;
+		Variant.VanillaPlatformId = MsgVariant.VanillaPlatformId;
+		Variant.PlatformDisplayName = FText::FromString(MsgVariant.PlatformDisplayName);
+	}
+
+	LastUpdateTime = FDateTime::UtcNow();
+}
+
+
+/* ITargetDeviceProxy interface
+ *****************************************************************************/
+
+int32 FTargetDeviceProxy::GetNumVariants() const
+{
+	return TargetDeviceVariants.Num();
+}
+
+
+int32 FTargetDeviceProxy::GetVariants(TArray<FName>& OutVariants) const
+{
+	return TargetDeviceVariants.GetKeys(OutVariants);
+}
+
+
+FName FTargetDeviceProxy::GetTargetDeviceVariant(const FString& InDeviceId) const
+{
+	for (TMap<FName, FTargetDeviceProxyVariant>::TConstIterator ItVariant(TargetDeviceVariants); ItVariant; ++ItVariant)
+	{
+		const FTargetDeviceProxyVariant& Variant = ItVariant.Value();
+
+		if (Variant.DeviceID == InDeviceId)
+		{
+			return ItVariant.Key();
+		}
+	}
+
+	return NAME_None;
+}
+
+
+const FString& FTargetDeviceProxy::GetTargetDeviceId(FName InVariant) const
+{
+	if (InVariant == NAME_None)
+	{
+		return TargetDeviceVariants[DefaultVariant].DeviceID;
+	}
+
+	return TargetDeviceVariants[InVariant].DeviceID;
+}
+
+
+FString FTargetDeviceProxy::GetTargetPlatformName(FName InVariant) const
+{
+	if (InVariant == NAME_None)
+	{
+		return TargetDeviceVariants[DefaultVariant].TargetPlatformName;
+	}
+
+	return TargetDeviceVariants[InVariant].TargetPlatformName;
+}
+
+
+FName FTargetDeviceProxy::GetTargetPlatformId(FName InVariant) const
+{
+	if (InVariant == NAME_None)
+	{
+		return TargetDeviceVariants[DefaultVariant].TargetPlatformId;
+	}
+
+	return TargetDeviceVariants[InVariant].TargetPlatformId;
+}
+
+
+FName FTargetDeviceProxy::GetVanillaPlatformId(FName InVariant) const
+{
+	if (InVariant == NAME_None)
+	{
+		return TargetDeviceVariants[DefaultVariant].VanillaPlatformId;
+	}
+
+	return TargetDeviceVariants[InVariant].VanillaPlatformId;
+}
+
+
+FText FTargetDeviceProxy::GetPlatformDisplayName(FName InVariant) const
+{
+	if (InVariant == NAME_None)
+	{
+		return TargetDeviceVariants[DefaultVariant].PlatformDisplayName;
+	}
+
+	return TargetDeviceVariants[InVariant].PlatformDisplayName;
+}
+
+
+bool FTargetDeviceProxy::HasDeviceId(const FString& InDeviceId) const
+{
+	for (TMap<FName, FTargetDeviceProxyVariant>::TConstIterator ItVariant(TargetDeviceVariants); ItVariant; ++ItVariant)
+	{
+		const FTargetDeviceProxyVariant& Variant = ItVariant.Value();
+
+		if (Variant.DeviceID == InDeviceId)
+		{
+			return true;
+		}
+	}
+
+	return false;
+}
+
+
+bool FTargetDeviceProxy::HasTargetPlatform(FName InTargetPlatformId) const
+{
+	for (TMap<FName, FTargetDeviceProxyVariant>::TConstIterator ItVariant(TargetDeviceVariants); ItVariant; ++ItVariant)
+	{
+		const FTargetDeviceProxyVariant& Variant = ItVariant.Value();
+
+		if (Variant.TargetPlatformId == InTargetPlatformId)
+		{
+			return true;
+		}
+	}
+
+	return false;
+}
+
+
+bool FTargetDeviceProxy::HasVariant(FName InVariant) const
+{
+	if (InVariant == NAME_None)
+	{
+		return TargetDeviceVariants.Contains(DefaultVariant);
+	}
+
+	return TargetDeviceVariants.Contains(InVariant);
+}
+
+
+bool FTargetDeviceProxy::DeployApp(FName InVariant, const TMap<FString, FString>& Files, const FGuid& TransactionId)
+{
+	for (TMap<FString, FString>::TConstIterator It(Files); It; ++It)
+	{
+		TSharedRef<IMessageAttachment, ESPMode::ThreadSafe> FileAttachment = MakeShareable(new FFileMessageAttachment(It.Key()));
+		FString SourcePath = It.Key();
+
+		MessageEndpoint->Send(new FTargetDeviceServiceDeployFile(It.Value(), TransactionId), FileAttachment, MessageAddress);
+	}
+
+	MessageEndpoint->Send(new FTargetDeviceServiceDeployCommit(InVariant, TransactionId), MessageAddress);
+
+	return true;
+}
+
+
+bool FTargetDeviceProxy::LaunchApp(FName InVariant, const FString& AppId, EBuildConfigurations::Type BuildConfiguration, const FString& Params)
+{
+	MessageEndpoint->Send(new FTargetDeviceServiceLaunchApp(InVariant, AppId, BuildConfiguration, Params), MessageAddress);
+
+	return true;
+}
+
+
+void FTargetDeviceProxy::PowerOff(bool Force)
+{
+	MessageEndpoint->Send(new FTargetDeviceServicePowerOff(FPlatformProcess::UserName(false), Force), MessageAddress);
+}
+
+
+void FTargetDeviceProxy::PowerOn()
+{
+	MessageEndpoint->Send(new FTargetDeviceServicePowerOn(FPlatformProcess::UserName(false)), MessageAddress);
+}
+
+
+void FTargetDeviceProxy::Reboot()
+{
+	MessageEndpoint->Send(new FTargetDeviceServiceReboot(FPlatformProcess::UserName(false)), MessageAddress);
+}
+
+
+void FTargetDeviceProxy::Run(FName InVariant, const FString& ExecutablePath, const FString& Params)
+{
+	MessageEndpoint->Send(new FTargetDeviceServiceRunExecutable(InVariant, ExecutablePath, Params), MessageAddress);
+}
+
+
+/* FTargetDeviceProxy implementation
+ *****************************************************************************/
+
+void FTargetDeviceProxy::InitializeMessaging()
+{
+	MessageEndpoint = FMessageEndpoint::Builder(FName(*FString::Printf(TEXT("FTargetDeviceProxy (%s)"), *Name)))
+		.Handling<FTargetDeviceServiceDeployFinished>(this, &FTargetDeviceProxy::HandleDeployFinishedMessage)
+		.Handling<FTargetDeviceServiceLaunchFinished>(this, &FTargetDeviceProxy::HandleLaunchFinishedMessage);
+}
+
+/* FTargetDeviceProxy event handlers
+*****************************************************************************/
+
+void FTargetDeviceProxy::HandleDeployFinishedMessage(const FTargetDeviceServiceDeployFinished& Message, const TSharedRef<IMessageContext, ESPMode::ThreadSafe>& Context)
+{
+	if (Message.Succeeded)
+	{
+		DeployCommittedDelegate.Broadcast(Message.TransactionId, Message.AppID);
+	}
+	else
+	{
+		DeployFailedDelegate.Broadcast(Message.TransactionId);
+	}
+}
+
+
+void FTargetDeviceProxy::HandleLaunchFinishedMessage(const FTargetDeviceServiceLaunchFinished& Message, const TSharedRef<IMessageContext, ESPMode::ThreadSafe>& Context)
+{
+	if (Message.Succeeded)
+	{
+		LaunchSucceededDelegate.Broadcast(Message.AppID, Message.ProcessId);
+	}
+	else
+	{
+		LaunchFailedDelegate.Broadcast(Message.AppID);
+	}
+}
diff --git a/Engine/Source/Developer/TargetDeviceServices/Private/Proxies/TargetDeviceProxy.h b/Engine/Source/Developer/TargetDeviceServices/Private/Proxies/TargetDeviceProxy.h
new file mode 100644
--- /dev/null
+++ b/Engine/Source/Developer/TargetDeviceServices/Private/Proxies/TargetDeviceProxy.h
@@ -0,0 +1,292 @@
+// Copyright 1998-2017 Epic Games, Inc. All Rights Reserved.
+
+#pragma once
+
+#include "CoreTypes.h"
+#include "IMessageContext.h"
+#include "MessageEndpoint.h"
+
+#include "ITargetDeviceProxy.h"
+
+class FTargetDeviceProxy;
+class IMessageContext;
+
+struct FTargetDeviceServiceDeployFinished;
+struct FTargetDeviceServiceLaunchFinished;
+struct FTargetDeviceServicePong;
+
+
+/** Type definition for shared references to instances of FTargetDeviceProxy. */
+typedef TSharedRef<class FTargetDeviceProxy> FTargetDeviceProxyRef;
+
+
+/**
+ * Implementation of the device proxy.
+ */
+class FTargetDeviceProxy
+	: public ITargetDeviceProxy
+{
+public:
+
+	/** Default constructor. */
+	FTargetDeviceProxy() { }
+
+	/**
+	 * Creates and initializes a new instance.
+	 *
+	 * @param InId The identifier of the target device to create this proxy for.
+	 */
+	FTargetDeviceProxy(const FString& InId);
+
+	/**
+	 * Creates and initializes a new instance.
+	 *
+	 * @param InId The identifier of the target device to create this proxy for.
+	 * @param Message The message to initialize from.
+	 * @param Context The message context.
+	 */
+	FTargetDeviceProxy(const FString& InId, const FTargetDeviceServicePong& Message, const TSharedRef<IMessageContext, ESPMode::ThreadSafe>& Context);
+
+public:
+
+	/**
+	 * Gets the time at which the proxy was last updated.
+	 *
+	 * @return Date and time of the last update.
+	 */
+	const FDateTime& GetLastUpdateTime() const
+	{
+		return LastUpdateTime;
+	}
+
+	/**
+	 * Updates the proxy's information from the given device service response.
+	 *
+	 * @param Message The message containing the response.
+	 * @param Context The message context.
+	 */
+	void UpdateFromMessage(const FTargetDeviceServicePong& Message, const TSharedRef<IMessageContext, ESPMode::ThreadSafe>& Context);
+
+public:
+
+	// ITargetDeviceProxy interface
+
+	virtual const bool CanMultiLaunch() const override
+	{
+		return SupportsMultiLaunch;
+	}
+
+	virtual bool CanPowerOff() const override
+	{
+		return SupportsPowerOff;
+	}
+
+	virtual bool CanPowerOn() const override
+	{
+		return SupportsPowerOn;
+	}
+
+	virtual bool CanReboot() const override
+	{
+		return SupportsReboot;
+	}
+
+	virtual bool CanSupportVariants() const override
+	{
+		return SupportsVariants;
+	}
+
+	virtual int32 GetNumVariants() const override;
+	virtual int32 GetVariants(TArray<FName>& OutVariants) const override;
+	virtual FName GetTargetDeviceVariant(const FString& InDeviceId) const override;
+	virtual const FString& GetTargetDeviceId(FName InVariant) const override;
+	virtual FString GetTargetPlatformName(FName InVariant) const override;
+	virtual FName GetTargetPlatformId(FName InVariant) const override;
+	virtual FName GetVanillaPlatformId(FName InVariant) const override;
+	virtual FText GetPlatformDisplayName(FName InVariant) const override;
+
+	virtual const FString& GetHostName() const override
+	{
+		return HostName;
+	}
+
+	virtual const FString& GetHostUser() const override
+	{
+		return HostUser;
+	}
+
+	virtual const FString& GetMake() const override
+	{
+		return Make;
+	}
+
+	virtual const FString& GetModel() const override
+	{
+		return Model;
+	}
+
+	virtual const FString& GetName() const override
+	{
+		return Name;
+	}
+
+	virtual const FString& GetDeviceUser() const override
+	{
+		return DeviceUser;
+	}
+
+	virtual const FString& GetDeviceUserPassword() const override
+	{
+		return DeviceUserPassword;
+	}
+
+	virtual const FString& GetType() const override
+	{
+		return Type;
+	}
+
+	virtual bool HasDeviceId(const FString& InDeviceId) const override;
+	virtual bool HasVariant(FName InVariant) const override;
+	virtual bool HasTargetPlatform(FName InTargetPlatformId) const override;
+
+	virtual bool IsConnected() const override
+	{
+		return Connected;
+	}
+
+	virtual bool IsShared() const override
+	{
+		return Shared;
+	}
+
+	virtual bool DeployApp(FName InVariant, const TMap<FString, FString>& Files, const FGuid& TransactionId) override;
+	virtual bool LaunchApp(FName InVariant, const FString& AppId, EBuildConfigurations::Type BuildConfiguration, const FString& Params) override;
+
+	virtual FOnTargetDeviceProxyDeployCommitted& OnDeployCommitted() override
+	{
+		return DeployCommittedDelegate;
+	}
+
+	virtual FOnTargetDeviceProxyDeployFailed& OnDeployFailed() override
+	{
+		return DeployFailedDelegate;
+	}
+
+	virtual FOnTargetDeviceProxyLaunchFailed& OnLaunchFailed() override
+	{
+		return LaunchFailedDelegate;
+	}
+
+	virtual void PowerOff(bool Force) override;
+	virtual void PowerOn() override;
+	virtual void Reboot() override;
+	virtual void Run(FName InVariant, const FString& ExecutablePath, const FString& Params) override;
+
+protected:
+
+	/** Initializes the message endpoint. */
+	void InitializeMessaging();
+
+private:
+
+	/** Handles FTargetDeviceServiceDeployFinishedMessage messages. */
+	void HandleDeployFinishedMessage(const FTargetDeviceServiceDeployFinished& Message, const TSharedRef<IMessageContext, ESPMode::ThreadSafe>& Context);
+
+	/** Handles FTargetDeviceServiceLaunchFinishedMessage messages. */
+	void HandleLaunchFinishedMessage(const FTargetDeviceServiceLaunchFinished& Message, const TSharedRef<IMessageContext, ESPMode::ThreadSafe>& Context);
+
+private:
+
+	/** Holds a flag indicating whether the device is connected. */
+	bool Connected;
+
+	/** Holds the name of the computer that hosts the device. */
+	FString HostName;
+
+	/** Holds the name of the user that owns the device. */
+	FString HostUser;
+
+	/** Holds the time at which the last ping reply was received. */
+	FDateTime LastUpdateTime;
+
+	/** Holds the device make. */
+	FString Make;
+
+	/** Holds the remote device's message bus address. */
+	FMessageAddress MessageAddress;
+
+	/** Holds the local message bus endpoint. */
+	TSharedPtr<FMessageEndpoint, ESPMode::ThreadSafe> MessageEndpoint;
+
+	/** Holds the remote device's model. */
+	FString Model;
+
+	/** Holds the name of the device. */
+	FString Name;
+
+	/** Holds device user. */
+	FString DeviceUser;
+
+	/** Holds device user password. */
+	FString DeviceUserPassword;
+
+	/** Holds a flag indicating whether the device is being shared with other users. */
+	bool Shared;
+
+	/** Holds a flag indicating whether the device supports multi-launch. */
+	bool SupportsMultiLaunch;
+
+	/** Holds a flag indicating whether the device can power off. */
+	bool SupportsPowerOff;
+
+	/** Holds a flag indicating whether the device can be power on. */
+	bool SupportsPowerOn;
+
+	/** Holds a flag indicating whether the device can reboot. */
+	bool SupportsReboot;
+
+	/** Holds a flag indicating whether the device's target platform supports variants. */
+	bool SupportsVariants;
+
+	/** Holds the device type. */
+	FString Type;
+
+	/** Holds default variant name. */
+	FName DefaultVariant;
+
+
+	/** Holds data about a device proxy variant. */
+	class FTargetDeviceProxyVariant
+	{
+	public:
+
+		/** Holds a string version of the variants device id. */
+		FString DeviceID;
+
+		/** Holds the variant name, this is the the map key as well. */
+		FName VariantName;
+
+		/** Platform information. */
+		FString TargetPlatformName;
+		FName TargetPlatformId;
+		FName VanillaPlatformId;
+		FText PlatformDisplayName;
+	};
+
+	/** Map of all the Variants for this Device. */
+	TMap<FName, FTargetDeviceProxyVariant> TargetDeviceVariants;
+
+private:
+
+	/** Holds a delegate to be invoked when a build has been deployed to the target device. */
+	FOnTargetDeviceProxyDeployCommitted DeployCommittedDelegate;
+
+	/** Holds a delegate to be invoked when a build has failed to deploy to the target device. */
+	FOnTargetDeviceProxyDeployFailed DeployFailedDelegate;
+
+	/** Holds a delegate to be invoked when a build has failed to launch on the target device. */
+	FOnTargetDeviceProxyLaunchFailed LaunchFailedDelegate;
+
+	/** Holds a delegate to be invoked when a build has succeeded to launch on the target device. */
+	FOnTargetDeviceProxyLaunchSucceeded LaunchSucceededDelegate;
+};
