Project Launcher now has 'Show Missing' option in the cooked maps section.

diff --git a/Engine/Source/Developer/ProjectLauncher/Private/Widgets/Cook/SProjectLauncherCookByTheBookSettings.cpp b/Engine/Source/Developer/ProjectLauncher/Private/Widgets/Cook/SProjectLauncherCookByTheBookSettings.cpp
--- a/Engine/Source/Developer/ProjectLauncher/Private/Widgets/Cook/SProjectLauncherCookByTheBookSettings.cpp
+++ b/Engine/Source/Developer/ProjectLauncher/Private/Widgets/Cook/SProjectLauncherCookByTheBookSettings.cpp
@@ -231,7 +231,7 @@ TSharedRef<SWidget> SProjectLauncherCookByTheBookSettings::MakeComplexWidget()
 							]
 
 							+ SHorizontalBox::Slot()
-								.FillWidth(1.0f)
+								.AutoWidth()
 								.Padding(8.0f, 0.0f, 0.0f, 0.0f)
 								[
 									// cooked maps radio button
@@ -244,6 +244,21 @@ TSharedRef<SWidget> SProjectLauncherCookByTheBookSettings::MakeComplexWidget()
 												.Text(LOCTEXT("CookedMapsCheckBoxText", "Show cooked"))
 										]
 								]
+
+							+ SHorizontalBox::Slot()
+								.FillWidth(1.0f)
+								.Padding(8.0f, 0.0f, 0.0f, 0.0f)
+								[
+									// missing maps radio button
+									SNew(SCheckBox)
+										.IsChecked(this, &SProjectLauncherCookByTheBookSettings::HandleShowCheckBoxIsChecked, EShowMapsChoices::ShowMissingMaps)
+										.OnCheckStateChanged(this, &SProjectLauncherCookByTheBookSettings::HandleShowCheckBoxCheckStateChanged, EShowMapsChoices::ShowMissingMaps)
+										.Style(FEditorStyle::Get(), "RadioButton")
+										[
+											SNew(STextBlock)
+												.Text(LOCTEXT("MissingMapsCheckBoxText", "Show missing"))
+										]
+								]
 						]
 
 					+ SVerticalBox::Slot()
@@ -944,7 +959,7 @@ TSharedRef<SWidget> SProjectLauncherCookByTheBookSettings::MakeSimpleWidget()
 								]
 
 							+ SHorizontalBox::Slot()
-								.FillWidth(1.0f)
+								.AutoWidth()
 								.Padding(8.0f, 0.0f, 0.0f, 0.0f)
 								[
 									// cooked maps radio button
@@ -957,6 +972,21 @@ TSharedRef<SWidget> SProjectLauncherCookByTheBookSettings::MakeSimpleWidget()
 												.Text(LOCTEXT("CookedMapsCheckBoxText", "Show cooked"))
 										]
 								]
+
+							+ SHorizontalBox::Slot()
+								.FillWidth(1.0f)
+								.Padding(8.0f, 0.0f, 0.0f, 0.0f)
+								[
+									// missing maps radio button
+									SNew(SCheckBox)
+										.IsChecked(this, &SProjectLauncherCookByTheBookSettings::HandleShowCheckBoxIsChecked, EShowMapsChoices::ShowMissingMaps)
+										.OnCheckStateChanged(this, &SProjectLauncherCookByTheBookSettings::HandleShowCheckBoxCheckStateChanged, EShowMapsChoices::ShowMissingMaps)
+										.Style(FEditorStyle::Get(), "RadioButton")
+										[
+											SNew(STextBlock)
+												.Text(LOCTEXT("MissingMapsCheckBoxText", "Show missing"))
+										]
+								]
 						]
 
 					+ SVerticalBox::Slot()
@@ -1209,13 +1239,26 @@ void SProjectLauncherCookByTheBookSettings::RefreshMapList()
 	{
 		TArray<FString> AvailableMaps = FGameProjectHelper::GetAvailableMaps(SelectedProfile->GetProjectBasePath(), SelectedProfile->SupportsEngineMaps(), true);
 
-		for (int32 AvailableMapIndex = 0; AvailableMapIndex < AvailableMaps.Num(); ++AvailableMapIndex)
+		if (ShowMapsChoice == EShowMapsChoices::ShowMissingMaps)
 		{
-			FString& Map = AvailableMaps[AvailableMapIndex];
-
-			if ((ShowMapsChoice == EShowMapsChoices::ShowAllMaps) || SelectedProfile->GetCookedMaps().Contains(Map))
+			for (const FString& Map : SelectedProfile->GetCookedMaps())
 			{
-				MapList.Add(MakeShareable(new FString(Map)));
+				if (!AvailableMaps.Contains(Map))
+				{
+					MapList.Add(MakeShareable(new FString(Map)));
+				}
+			}
+		}
+		else
+		{
+			for (int32 AvailableMapIndex = 0; AvailableMapIndex < AvailableMaps.Num(); ++AvailableMapIndex)
+			{
+				FString& Map = AvailableMaps[AvailableMapIndex];
+
+				if ((ShowMapsChoice == EShowMapsChoices::ShowAllMaps) || SelectedProfile->GetCookedMaps().Contains(Map))
+				{
+					MapList.Add(MakeShareable(new FString(Map)));
+				}
 			}
 		}
 	}
@@ -1470,6 +1513,10 @@ FText SProjectLauncherCookByTheBookSettings::HandleNoMapsTextBlockText() const
 		{
 			return LOCTEXT("NoMapsSelectedText", "No map selected. Only startup packages will be cooked!");
 		}
+		else if (ShowMapsChoice == EShowMapsChoices::ShowMissingMaps)
+		{
+			return LOCTEXT("NoMissingMapsFoundText", "No missing maps were found.");
+		}
 	}
 
 	return FText();
diff --git a/Engine/Source/Developer/ProjectLauncher/Private/Widgets/Cook/SProjectLauncherCookByTheBookSettings.h b/Engine/Source/Developer/ProjectLauncher/Private/Widgets/Cook/SProjectLauncherCookByTheBookSettings.h
--- a/Engine/Source/Developer/ProjectLauncher/Private/Widgets/Cook/SProjectLauncherCookByTheBookSettings.h
+++ b/Engine/Source/Developer/ProjectLauncher/Private/Widgets/Cook/SProjectLauncherCookByTheBookSettings.h
@@ -23,7 +23,10 @@ namespace EShowMapsChoices
 		ShowAllMaps,
 
 		/** Only show maps that are to be cooked. */
-		ShowCookedMaps
+		ShowCookedMaps,
+
+		/** Only show maps that are to be cooked but missing. */
+		ShowMissingMaps,
 	};
 }
 
