Removing outdated styling that was causing the Project Launcher UI to render incorrectly

diff --git a/Engine/Source/Developer/ProjectLauncher/Private/Widgets/Deploy/SProjectLauncherSimpleDeviceListView.cpp b/Engine/Source/Developer/ProjectLauncher/Private/Widgets/Deploy/SProjectLauncherSimpleDeviceListView.cpp
--- a/Engine/Source/Developer/ProjectLauncher/Private/Widgets/Deploy/SProjectLauncherSimpleDeviceListView.cpp
+++ b/Engine/Source/Developer/ProjectLauncher/Private/Widgets/Deploy/SProjectLauncherSimpleDeviceListView.cpp
@@ -74,7 +74,6 @@ void SProjectLauncherSimpleDeviceListView::Construct(const FArguments& InArgs, c
 			.HAlign(HAlign_Left)
 			[
 				SNew(STextBlock)
-				.TextStyle(FCoreStyle::Get(), "ToolBar.Keybinding")
 				.Text(LOCTEXT("ProjectLauncherDeviceManagerLinkPreamble", "Don't see your device? Verify it's setup and claimed in the "))
 			]
 
@@ -85,7 +84,6 @@ void SProjectLauncherSimpleDeviceListView::Construct(const FArguments& InArgs, c
 			[
 				// all cultures hyper link
 				SNew(SHyperlink)
-				.TextStyle(FCoreStyle::Get(), "ToolBar.Keybinding")
 				.OnNavigate(this, &SProjectLauncherSimpleDeviceListView::HandleDeviceManagerHyperlinkNavigate)
 				.Text(LOCTEXT("ProjectLauncherDeviceManagerLink", "Device Manager."))
 				.ToolTipText(LOCTEXT("ProjectLauncherDeviceManagerLinkTooltip", "Open the Device Manager window, where you can setup and claim devices connected to your machine or shared on the network."))
diff --git a/Engine/Source/Developer/ProjectLauncher/Private/Widgets/Profile/SProjectLauncherProfileListRow.h b/Engine/Source/Developer/ProjectLauncher/Private/Widgets/Profile/SProjectLauncherProfileListRow.h
--- a/Engine/Source/Developer/ProjectLauncher/Private/Widgets/Profile/SProjectLauncherProfileListRow.h
+++ b/Engine/Source/Developer/ProjectLauncher/Private/Widgets/Profile/SProjectLauncherProfileListRow.h
@@ -97,8 +97,6 @@ public:
 					.Padding(4, 0, 0, 0)
 					[
 						SNew(SButton)
-						.ButtonStyle(FCoreStyle::Get(), "Toolbar.Button")
-						.ForegroundColor(FEditorStyle::GetSlateColor("DefaultForeground"))
 						.OnClicked(this, &SProjectLauncherProfileListRow::OnEditClicked)
 						.ToolTipText(LOCTEXT("EditProfileToolTipText", "Edit profile."))
 						.HAlign(HAlign_Center)
diff --git a/Engine/Source/Developer/ProjectLauncher/Private/Widgets/SProjectLauncher.cpp b/Engine/Source/Developer/ProjectLauncher/Private/Widgets/SProjectLauncher.cpp
--- a/Engine/Source/Developer/ProjectLauncher/Private/Widgets/SProjectLauncher.cpp
+++ b/Engine/Source/Developer/ProjectLauncher/Private/Widgets/SProjectLauncher.cpp
@@ -1,13 +1,10 @@
 // Copyright Epic Games, Inc. All Rights Reserved.
 
 #include "SProjectLauncher.h"
-
-#include "EditorStyleSet.h"
 #include "Framework/Commands/UIAction.h"
 #include "Framework/Docking/TabManager.h"
 #include "Framework/MultiBox/MultiBoxBuilder.h"
 #include "SlateOptMacros.h"
-#include "Styling/CoreStyle.h"
 #include "Styling/SlateTypes.h"
 #include "Widgets/SBoxPanel.h"
 #include "Widgets/Images/SImage.h"
@@ -25,6 +22,7 @@
 #include "Widgets/Progress/SProjectLauncherProgress.h"
 #include "Widgets/Project/SProjectLauncherProjectPicker.h"
 #include "Widgets/Settings/SProjectLauncherSettings.h"
+#include "Styling/AppStyle.h"
 
 
 #define LOCTEXT_NAMESPACE "SProjectLauncher"
@@ -121,10 +119,10 @@ void SProjectLauncher::Construct(const FArguments& InArgs, const TSharedRef<SDoc
 					[
 						SNew(SBorder)
 						.Padding(2)
-						.BorderImage(FEditorStyle::GetBrush("ToolPanel.GroupBorder"))
+						.BorderImage(FAppStyle::Get().GetBrush("ToolPanel.GroupBorder"))
 						[
 							SNew(SCheckBox)
-							.Style(FEditorStyle::Get(), "ToggleButtonCheckbox")
+							.Style(FAppStyle::Get(), "ToggleButtonCheckbox")
 							.IsFocusable(true)
 							.ToolTipText(LOCTEXT("ToggleAdvancedOptionsToolTipText", "Toggles Advanced Options"))
 							.OnCheckStateChanged(this, &SProjectLauncher::OnAdvancedChanged)
@@ -147,8 +145,6 @@ void SProjectLauncher::Construct(const FArguments& InArgs, const TSharedRef<SDoc
 								.Padding(4,0,4,0)
 								[
 									SNew(STextBlock)
-									.TextStyle(FCoreStyle::Get(), "Toolbar.Label")
-									.ShadowOffset(FVector2D::UnitVector)
 									.Text(LOCTEXT("AdvancedButton", "Advanced"))
 								]
 							]
@@ -181,7 +177,7 @@ void SProjectLauncher::Construct(const FArguments& InArgs, const TSharedRef<SDoc
 					.Padding(2)
 					[
 						SNew(SBorder)
-						.BorderImage(FEditorStyle::GetBrush("ToolPanel.GroupBorder"))
+						.BorderImage(FAppStyle::Get().GetBrush("ToolPanel.GroupBorder"))
 						.Padding(4.0)
 						[
 							SNew(SHorizontalBox)
@@ -190,8 +186,6 @@ void SProjectLauncher::Construct(const FArguments& InArgs, const TSharedRef<SDoc
 							.FillWidth(1.0f)
 							[
 								SNew(STextBlock)
-								.TextStyle(FCoreStyle::Get(), "Toolbar.Label")
-								.ShadowOffset(FVector2D::UnitVector)
 								.Text(LOCTEXT("ProjectLauncherCustomProfilesTitle", "Custom Launch Profiles"))
 							]
 
@@ -200,8 +194,6 @@ void SProjectLauncher::Construct(const FArguments& InArgs, const TSharedRef<SDoc
 							.AutoWidth()
 							[
 								SNew(SComboButton)
-								.ComboButtonStyle(FEditorStyle::Get(), "GenericFilters.ComboButtonStyle")
-								.ForegroundColor(FLinearColor::White)
 								.ContentPadding(0)
 								.ToolTipText(LOCTEXT("AddFilterToolTip", "Add a new custom launch profile using wizard"))
 								.OnGetMenuContent(this, &SProjectLauncher::MakeProfileWizardsMenu)
@@ -215,14 +207,12 @@ void SProjectLauncher::Construct(const FArguments& InArgs, const TSharedRef<SDoc
 							.AutoWidth()
 							[
 								SNew(SButton)
-								.ButtonStyle(FEditorStyle::Get(), "ToggleButton")
 								.ToolTipText(LOCTEXT("ProjectLauncherCustomProfileAdd", "Add a new custom launch profile."))
 								.ContentPadding(0)
 								.OnClicked(this, &SProjectLauncher::OnAddCustomLaunchProfileClicked)
 								[
 									SNew(SImage)
-									.Image(FCoreStyle::Get().GetBrush("EditableComboBox.Add"))
-									.ColorAndOpacity(FSlateColor(FLinearColor::White))
+									.Image(FAppStyle::Get().GetBrush("EditableComboBox.Add"))
 								]
 							]
 						]
@@ -232,7 +222,7 @@ void SProjectLauncher::Construct(const FArguments& InArgs, const TSharedRef<SDoc
 					.Padding(2)
 					[
 						SAssignNew(ProfileList, SBorder)
-						.BorderImage(FEditorStyle::GetBrush("NoBorder"))
+						.BorderImage(FAppStyle::Get().GetBrush("NoBorder"))
 						.Padding(0)
 						[
 							// Simple Launch List
@@ -298,7 +288,7 @@ ECheckBoxState SProjectLauncher::OnIsAdvanced() const
 
 const FSlateBrush* SProjectLauncher::GetAdvancedToggleBrush() const
 {
-	return FEditorStyle::GetBrush("LauncherCommand.AdvancedBuild.Medium");
+	return FAppStyle::Get().GetBrush("LauncherCommand.AdvancedBuild.Medium");
 }
 
 
diff --git a/Engine/Source/Developer/ProjectLauncher/Private/Widgets/Shared/SProjectLauncherProfileLaunchButton.h b/Engine/Source/Developer/ProjectLauncher/Private/Widgets/Shared/SProjectLauncherProfileLaunchButton.h
--- a/Engine/Source/Developer/ProjectLauncher/Private/Widgets/Shared/SProjectLauncherProfileLaunchButton.h
+++ b/Engine/Source/Developer/ProjectLauncher/Private/Widgets/Shared/SProjectLauncherProfileLaunchButton.h
@@ -8,9 +8,8 @@
 #include "Widgets/DeclarativeSyntaxSupport.h"
 #include "Widgets/SWidget.h"
 #include "Widgets/SCompoundWidget.h"
-#include "Styling/CoreStyle.h"
+#include "Styling/AppStyle.h"
 #include "Widgets/SBoxPanel.h"
-#include "EditorStyleSet.h"
 #include "Widgets/Text/STextBlock.h"
 #include "Widgets/Images/SImage.h"
 #include "Widgets/Input/SButton.h"
@@ -57,8 +56,6 @@ public:
 			.VAlign(VAlign_Center)
 			[
 				SNew(SButton)
-				.ButtonStyle(FCoreStyle::Get(), "Toolbar.Button")
-				.ForegroundColor(FEditorStyle::GetSlateColor("DefaultForeground"))
 				.OnClicked(InArgs._OnClicked)
 				.HAlign(HAlign_Center)
 				.VAlign(VAlign_Center)
@@ -96,8 +93,6 @@ public:
 			.AutoHeight()
 			[
 				SNew(STextBlock)
-				.TextStyle(FCoreStyle::Get(), "Toolbar.Label")
-				.ShadowOffset(FVector2D::UnitVector)
 				.Text(LOCTEXT("ProjectLauncherLaunch", "Launch"))
 			];
 		}
@@ -149,13 +144,13 @@ private:
 	// Get the SlateIcon for Launch Button
 	const FSlateBrush* GetLaunchIcon() const
 	{
-		return FEditorStyle::GetBrush("Launcher.Run");
+		return FAppStyle::Get().GetBrush("Launcher.Run");
 	}
 
 	// Get the SlateIcon for Error Image
 	const FSlateBrush* GetErrorIcon() const
 	{
-		return FEditorStyle::GetBrush(TEXT("Icons.Error"));
+		return FAppStyle::Get().GetBrush(TEXT("Icons.Error"));
 	}
 
 	// Callback to see if the error icon should be displayed, based on validity of the launch profile
