Fix for duplicate localization key warnings and missing localization defines/undefines in various files.
#rnx

diff --git a/Engine/Plugins/Animation/ControlRig/Source/ControlRigEditor/Private/Graph/SControlRigGraphChangePinType.cpp b/Engine/Plugins/Animation/ControlRig/Source/ControlRigEditor/Private/Graph/SControlRigGraphChangePinType.cpp
--- a/Engine/Plugins/Animation/ControlRig/Source/ControlRigEditor/Private/Graph/SControlRigGraphChangePinType.cpp
+++ b/Engine/Plugins/Animation/ControlRig/Source/ControlRigEditor/Private/Graph/SControlRigGraphChangePinType.cpp
@@ -366,7 +366,7 @@ void SControlRigChangePinType::FillPinTypeMenu(FMenuBuilder& MenuBuilder)
 						.Text(GetBindingText(Type))
 						.ColorAndOpacity(FLinearColor::White * (bIsFilteredOut ? 0.5f : 1.f))
 						.ToolTipText(bIsFilteredOut ? LOCTEXT("WildcardAvailableTypeTooltip","Will break connections if resolved to this type.") :
-							LOCTEXT("WildcardAvailableTypeTooltip", "Available filtered type"))
+							LOCTEXT("WildcardAvailableFilterTypeTooltip", "Available filtered type"))
 					]);
 		}
 	}
@@ -415,4 +415,4 @@ TSharedRef<SWidget>	SControlRigGraphChangePinType::GetDefaultValueWidget()
 		.ModelPins(ModelPins);
 }
 
-#undef LOCTEXT_NAMESPACE
\ No newline at end of file
+#undef LOCTEXT_NAMESPACE
diff --git a/Engine/Plugins/Animation/GameplayInsights/Source/RewindDebugger/Private/RewindDebuggerModule.cpp b/Engine/Plugins/Animation/GameplayInsights/Source/RewindDebugger/Private/RewindDebuggerModule.cpp
--- a/Engine/Plugins/Animation/GameplayInsights/Source/RewindDebugger/Private/RewindDebuggerModule.cpp
+++ b/Engine/Plugins/Animation/GameplayInsights/Source/RewindDebugger/Private/RewindDebuggerModule.cpp
@@ -162,7 +162,7 @@ void FRewindDebuggerModule::StartupModule()
 		.SetGroup(WorkspaceMenu::GetMenuStructure().GetDeveloperToolsDebugCategory())
 		.SetDisplayName(LOCTEXT("DetailsTabTitle", "Rewind Debugger Details"))
 		.SetIcon(FSlateIcon("RewindDebuggerStyle", "RewindDebugger.RewindDetailsIcon"))
-		.SetTooltipText(LOCTEXT("TooltipText", "Opens Rewind Debugger Details Window."));
+		.SetTooltipText(LOCTEXT("DetailsWindowTooltipText", "Opens Rewind Debugger Details Window."));
 
 	RewindDebuggerCameraExtension.Initialize();
 	IModularFeatures::Get().RegisterModularFeature(IRewindDebuggerExtension::ModularFeatureName, &RewindDebuggerCameraExtension);
diff --git a/Engine/Plugins/Animation/IKRig/Source/IKRig/Private/IKRigProcessor.cpp b/Engine/Plugins/Animation/IKRig/Source/IKRig/Private/IKRigProcessor.cpp
--- a/Engine/Plugins/Animation/IKRig/Source/IKRig/Private/IKRigProcessor.cpp
+++ b/Engine/Plugins/Animation/IKRig/Source/IKRig/Private/IKRigProcessor.cpp
@@ -39,7 +39,7 @@ void UIKRigProcessor::Initialize(const UIKRigDefinition* InRigAsset, const USkel
 	if (!UIKRigProcessor::IsIKRigCompatibleWithSkeleton(InRigAsset, SkeletalMesh))
 	{
 		InRigAsset->Log.LogError(FText::Format(
-			LOCTEXT("NoSkeleton", "Trying to initialize IKRig, '{0}' with a Skeleton that is missing required bones. See prior warnings."),
+			LOCTEXT("SkeletonMissingRequiredBones", "Trying to initialize IKRig, '{0}' with a Skeleton that is missing required bones. See prior warnings."),
 			FText::FromString(InRigAsset->GetName())));
 		return;
 	}
@@ -183,7 +183,7 @@ bool UIKRigProcessor::IsIKRigCompatibleWithSkeleton(const UIKRigDefinition* InRi
 				bAllParentsValid = false;
 
 				InRigAsset->Log.LogError(FText::Format(
-			LOCTEXT("MissingParent", "IK Rig is running on a skeleton with a required bone, '{0}', that expected to have a valid parent. The expected parent was, '{1}'."),
+			LOCTEXT("InvalidParent", "IK Rig is running on a skeleton with a required bone, '{0}', that expected to have a valid parent. The expected parent was, '{1}'."),
 				FText::FromName(RequiredBone),
 				FText::FromName(AssetParentName)));
 				
@@ -195,7 +195,7 @@ bool UIKRigProcessor::IsIKRigCompatibleWithSkeleton(const UIKRigDefinition* InRi
 			{
 				// we only warn about this, because it may be nice not to have the exact same hierarchy
 				InRigAsset->Log.LogWarning(FText::Format(
-				LOCTEXT("MissingParent", "IK Rig is running on a skeleton with a required bone, '{0}', that has a different parent '{1}'. The expected parent was, '{2}'."),
+				LOCTEXT("UnexpectedParent", "IK Rig is running on a skeleton with a required bone, '{0}', that has a different parent '{1}'. The expected parent was, '{2}'."),
 				FText::FromName(RequiredBone),
 				FText::FromName(InputParentName),
 				FText::FromName(AssetParentName)));
@@ -411,4 +411,4 @@ void UIKRigProcessor::ResolveFinalGoalTransforms(const FTransform& WorldToCompon
 	}
 }
 
-#undef LOCTEXT_NAMESPACE
\ No newline at end of file
+#undef LOCTEXT_NAMESPACE
diff --git a/Engine/Plugins/Animation/IKRig/Source/IKRig/Private/Retargeter/IKRetargetProcessor.cpp b/Engine/Plugins/Animation/IKRig/Source/IKRig/Private/Retargeter/IKRetargetProcessor.cpp
--- a/Engine/Plugins/Animation/IKRig/Source/IKRig/Private/Retargeter/IKRetargetProcessor.cpp
+++ b/Engine/Plugins/Animation/IKRig/Source/IKRig/Private/Retargeter/IKRetargetProcessor.cpp
@@ -1,4 +1,4 @@
-﻿// Copyright Epic Games, Inc. All Rights Reserved.
+// Copyright Epic Games, Inc. All Rights Reserved.
 
 #include "Retargeter/IKRetargetProcessor.h"
 
@@ -1428,7 +1428,7 @@ bool UIKRetargetProcessor::InitializeIKRig(UObject* Outer, const USkeletalMesh*
 		if (!IKRigProcessor->GetGoalContainer().FindGoalByName(ChainPair.IKGoalName))
 		{
 			RetargeterAsset->Log.LogError( FText::Format(
-			LOCTEXT("CountNotFindRootBone", "IK Retargeter has target bone chain, {0} that references an IK Goal, {1} that is not present in any of the solvers in the IK Rig asset."),
+			LOCTEXT("TargetIKBoneNotInSolver", "IK Retargeter has target bone chain, {0} that references an IK Goal, {1} that is not present in any of the solvers in the IK Rig asset."),
 			FText::FromName(ChainPair.TargetBoneChainName), FText::FromName(ChainPair.IKGoalName)));
 			return false;
 		}
@@ -1653,4 +1653,4 @@ void UIKRetargetProcessor::CopyAllSettingsFromAsset()
 }
 #endif
 
-#undef LOCTEXT_NAMESPACE
\ No newline at end of file
+#undef LOCTEXT_NAMESPACE
diff --git a/Engine/Plugins/Animation/IKRig/Source/IKRigEditor/Private/RetargetEditor/IKRetargetDetails.cpp b/Engine/Plugins/Animation/IKRig/Source/IKRigEditor/Private/RetargetEditor/IKRetargetDetails.cpp
--- a/Engine/Plugins/Animation/IKRig/Source/IKRigEditor/Private/RetargetEditor/IKRetargetDetails.cpp
+++ b/Engine/Plugins/Animation/IKRig/Source/IKRigEditor/Private/RetargetEditor/IKRetargetDetails.cpp
@@ -1,4 +1,4 @@
-﻿// Copyright Epic Games, Inc. All Rights Reserved.
+// Copyright Epic Games, Inc. All Rights Reserved.
 
 #include "RetargetEditor/IKRetargetDetails.h"
 
@@ -37,7 +37,7 @@ void FIKRetargeterDetails::CustomizeDetails(IDetailLayoutBuilder& DetailBuilder)
 	TSharedPtr<FUICommandList> Commands = Controller->GetEditorController()->Editor.Pin()->GetToolkitCommands();
 
 	// add a new category at the top to edit the retarget pose
-	IDetailCategoryBuilder& EditPoseCategoryBuilder = DetailBuilder.EditCategory( "Edit Retarget Pose", LOCTEXT("EditPoseLabel", "Edit Retarget Pose"), ECategoryPriority::Default );
+	IDetailCategoryBuilder& EditPoseCategoryBuilder = DetailBuilder.EditCategory( "Edit Retarget Pose", LOCTEXT("EditRetargetPoseLabel", "Edit Retarget Pose"), ECategoryPriority::Default );
 
 	// fill list of pose names
 	PoseNames.Reset();
@@ -69,7 +69,7 @@ void FIKRetargeterDetails::CustomizeDetails(IDetailLayoutBuilder& DetailBuilder)
 	];
 
 	// add pose editing toolbar
-	FDetailWidgetRow& ToolbarRow = EditPoseCategoryBuilder.AddCustomRow(LOCTEXT("CurrentPoseLabel", "Edit Pose"))
+	FDetailWidgetRow& ToolbarRow = EditPoseCategoryBuilder.AddCustomRow(LOCTEXT("EditPoseToolbarLabel", "Edit Pose"))
 	.WholeRowWidget
 	[
 		SNew(SHorizontalBox)
diff --git a/Engine/Plugins/Animation/IKRig/Source/IKRigEditor/Private/RetargetEditor/IKRetargetEditorController.cpp b/Engine/Plugins/Animation/IKRig/Source/IKRigEditor/Private/RetargetEditor/IKRetargetEditorController.cpp
--- a/Engine/Plugins/Animation/IKRig/Source/IKRigEditor/Private/RetargetEditor/IKRetargetEditorController.cpp
+++ b/Engine/Plugins/Animation/IKRig/Source/IKRigEditor/Private/RetargetEditor/IKRetargetEditorController.cpp
@@ -572,7 +572,7 @@
 					.TextStyle( FAppStyle::Get(), "DialogButtonText" )
 					.HAlign(HAlign_Center)
 					.VAlign(VAlign_Center)
-					.Text(LOCTEXT("ImportButtonLabel", "Import New Retarget Pose"))
+					.Text(LOCTEXT("ImportNewRetargetPoseButtonLabel", "Import New Retarget Pose"))
 					.OnClicked(this, &FIKRetargetEditorController::ImportRetargetPose)
 				]
 				
@@ -675,7 +675,7 @@
 
 	// create pop-up window for user to select animation sequence asset to import as a retarget pose]
 	ImportPoseFromSequenceWindow = SNew(SWindow)
-	.Title(LOCTEXT("ImportRetargetPose", "Import Retarget Pose from Sequence Asset"))
+	.Title(LOCTEXT("ImportRetargetPoseFromSequenceAsset", "Import Retarget Pose from Sequence Asset"))
 	.ClientSize(FVector2D(500, 600))
 	.SupportsMinimize(false) .SupportsMaximize(false)
 	[
@@ -767,7 +767,7 @@
 					.TextStyle( FAppStyle::Get(), "DialogButtonText" )
 					.HAlign(HAlign_Center)
 					.VAlign(VAlign_Center)
-					.Text(LOCTEXT("ImportButtonLabel", "Import As Retarget Pose"))
+					.Text(LOCTEXT("ImportAsRetargetPoseButtonLabel", "Import As Retarget Pose"))
 					.OnClicked(this, &FIKRetargetEditorController::OnImportPoseFromSequence)
 				]
 				
diff --git a/Engine/Plugins/Developer/Concert/ConcertApp/MultiUserServer/Source/MultiUserServer/Private/Widgets/HistoryDeletion/SDeleteActivityDependenciesDialog.cpp b/Engine/Plugins/Developer/Concert/ConcertApp/MultiUserServer/Source/MultiUserServer/Private/Widgets/HistoryDeletion/SDeleteActivityDependenciesDialog.cpp
--- a/Engine/Plugins/Developer/Concert/ConcertApp/MultiUserServer/Source/MultiUserServer/Private/Widgets/HistoryDeletion/SDeleteActivityDependenciesDialog.cpp
+++ b/Engine/Plugins/Developer/Concert/ConcertApp/MultiUserServer/Source/MultiUserServer/Private/Widgets/HistoryDeletion/SDeleteActivityDependenciesDialog.cpp
@@ -99,9 +99,9 @@ TSharedRef<SWidget> SDeleteActivityDependenciesDialog::CreateBody(const FGuid& I
 					const bool bWillBeDeleted = ActivitiesToDelete[ActivityID];
 					if (bWillBeDeleted)
 					{
-						return LOCTEXT("DeleteActivity.CheckBox.TooltipEnabled", "This is a possible dependency.\nWill be deleted.");
+						return LOCTEXT("DeleteActivity.CheckBox.TooltipEnabled.Deleted", "This is a possible dependency.\nWill be deleted.");
 					}
-					return LOCTEXT("DeleteActivity.CheckBox.TooltipEnabled", "This is a possible dependency.\nWill not be deleted.");
+					return LOCTEXT("DeleteActivity.CheckBox.TooltipEnabled.NotDeleted", "This is a possible dependency.\nWill not be deleted.");
 				});
 		});
 }
diff --git a/Engine/Plugins/Editor/DisplayClusterLaunch/Source/DisplayClusterLaunchEditor/Private/DisplayClusterLaunchEditorModule.cpp b/Engine/Plugins/Editor/DisplayClusterLaunch/Source/DisplayClusterLaunchEditor/Private/DisplayClusterLaunchEditorModule.cpp
--- a/Engine/Plugins/Editor/DisplayClusterLaunch/Source/DisplayClusterLaunchEditor/Private/DisplayClusterLaunchEditorModule.cpp
+++ b/Engine/Plugins/Editor/DisplayClusterLaunch/Source/DisplayClusterLaunchEditor/Private/DisplayClusterLaunchEditorModule.cpp
@@ -1088,7 +1088,7 @@ void FDisplayClusterLaunchEditorModule::AddConsoleVariablesEditorAssetsToToolbar
 					{
 						const FText ConsoleVariablesAssetName = FText::FromName(Asset.AssetName);
 						const FText ConsoleVariablesAssetTooltip =
-							FText::Format(LOCTEXT("SelectConsoleVariablesAssetFormat","Select Console Variables Asset '{0}'"), ConsoleVariablesAssetName);
+							FText::Format(LOCTEXT("SelectConsoleVariablesAssetFormatLong","Select Console Variables Asset '{0}'"), ConsoleVariablesAssetName);
 						NewMenuBuilder.AddMenuEntry(
 							ConsoleVariablesAssetName,
 							ConsoleVariablesAssetTooltip,
diff --git a/Engine/Plugins/Experimental/Animation/PoseSearch/Source/Editor/Private/PoseSearchDatabaseEditorToolkit.cpp b/Engine/Plugins/Experimental/Animation/PoseSearch/Source/Editor/Private/PoseSearchDatabaseEditorToolkit.cpp
--- a/Engine/Plugins/Experimental/Animation/PoseSearch/Source/Editor/Private/PoseSearchDatabaseEditorToolkit.cpp
+++ b/Engine/Plugins/Experimental/Animation/PoseSearch/Source/Editor/Private/PoseSearchDatabaseEditorToolkit.cpp
@@ -408,7 +408,7 @@ namespace UE::PoseSearch
 		check(Args.GetTabId() == FDatabaseEditorTabs::AssetDetailsID);
 
 		return SNew(SDockTab)
-			.Label(LOCTEXT("AssetDetails_Title", "Database Details"))
+			.Label(LOCTEXT("Database_Details_Title", "Database Details"))
 			[
 				EditingAssetWidget.ToSharedRef()
 			];
@@ -460,7 +460,7 @@ namespace UE::PoseSearch
 		}
 
 		return SNew(SDockTab)
-			.Label(LOCTEXT("AssetDetails_Title", "Selection Details"))
+			.Label(LOCTEXT("SelectionDetails_Title", "Selection Details"))
 			[
 				SNew(SScrollBox)
 				+SScrollBox::Slot()
diff --git a/Engine/Plugins/Experimental/Animation/PoseSearch/Source/Editor/Private/SPoseSearchDatabaseAssetList.cpp b/Engine/Plugins/Experimental/Animation/PoseSearch/Source/Editor/Private/SPoseSearchDatabaseAssetList.cpp
--- a/Engine/Plugins/Experimental/Animation/PoseSearch/Source/Editor/Private/SPoseSearchDatabaseAssetList.cpp
+++ b/Engine/Plugins/Experimental/Animation/PoseSearch/Source/Editor/Private/SPoseSearchDatabaseAssetList.cpp
@@ -273,8 +273,8 @@ namespace UE::PoseSearch
 			EUserInterfaceActionType::Button);
 
 		AddOptions.AddMenuEntry(
-			LOCTEXT("AddBlendSpace", "Add Blend Space"),
-			LOCTEXT("AddBlendSpaceTooltip", "Add new blend space to this group"),
+			LOCTEXT("AddBlendSpaceOption", "Add Blend Space"),
+			LOCTEXT("AddBlendSpaceOptionTooltip", "Add new blend space to this group"),
 			FSlateIcon(),
 			FUIAction(FExecuteAction::CreateSP(this, &SDatabaseAssetListItem::OnAddBlendSpace)),
 			NAME_None,
@@ -699,16 +699,16 @@ namespace UE::PoseSearch
 
 		AddOptions.BeginSection("AddOptions", LOCTEXT("AssetAddOptions", "Assets"));
 		AddOptions.AddMenuEntry(
-			LOCTEXT("AddSequence", "Sequence"),
-			LOCTEXT("AddSequenceTooltip", "Add new sequence to the default group"),
+			LOCTEXT("AddSequenceOption", "Sequence"),
+			LOCTEXT("AddSequenceOptionTooltip", "Add new sequence to the default group"),
 			FSlateIcon(),
 			FUIAction(FExecuteAction::CreateSP(this, &SDatabaseAssetTree::OnAddSequence, true)),
 			NAME_None,
 			EUserInterfaceActionType::Button);
 
 		AddOptions.AddMenuEntry(
-			LOCTEXT("AddBlendSpace", "Blend Space"),
-			LOCTEXT("AddBlendSpaceTooltip", "Add new blend space to the default group"),
+			LOCTEXT("BlendSpaceOption", "Blend Space"),
+			LOCTEXT("AddBlendSpaceToDefaultGroupTooltip", "Add new blend space to the default group"),
 			FSlateIcon(),
 			FUIAction(FExecuteAction::CreateSP(this, &SDatabaseAssetTree::OnAddBlendSpace, true)),
 			NAME_None,
@@ -717,8 +717,8 @@ namespace UE::PoseSearch
 
 		AddOptions.BeginSection("AddOptions", LOCTEXT("GroupAddOptions", "Groups"));
 		AddOptions.AddMenuEntry(
-			LOCTEXT("AddBlendSpace", "Group"),
-			LOCTEXT("AddBlendSpaceTooltip", "Add new group"),
+			LOCTEXT("GroupLabel", "Group"),
+			LOCTEXT("AddNewGroupTooltip", "Add new group"),
 			FSlateIcon(),
 			FUIAction(FExecuteAction::CreateSP(this, &SDatabaseAssetTree::OnAddGroup, true)),
 			NAME_None,
@@ -776,7 +776,7 @@ namespace UE::PoseSearch
 
 	void SDatabaseAssetTree::OnAddBlendSpace(bool bFinalizeChanges)
 	{
-		FScopedTransaction Transaction(LOCTEXT("AddBlendSpace", "Add Blend Space"));
+		FScopedTransaction Transaction(LOCTEXT("AddBlendSpaceTransaction", "Add Blend Space"));
 
 		EditorViewModel.Pin()->AddBlendSpaceToDatabase(nullptr, INDEX_NONE);
 
diff --git a/Engine/Plugins/Experimental/Chooser/Source/ChooserEditor/Private/ChooserTableEditor.cpp b/Engine/Plugins/Experimental/Chooser/Source/ChooserEditor/Private/ChooserTableEditor.cpp
--- a/Engine/Plugins/Experimental/Chooser/Source/ChooserEditor/Private/ChooserTableEditor.cpp
+++ b/Engine/Plugins/Experimental/Chooser/Source/ChooserEditor/Private/ChooserTableEditor.cpp
@@ -442,7 +442,7 @@ public:
 									.ComboButtonStyle(FAppStyle::Get(), "SimpleComboButton")
 									.ButtonContent()
 									[
-										SNew(STextBlock).Text(LOCTEXT("Add", "+ Add Row"))
+										SNew(STextBlock).Text(LOCTEXT("AddRow", "+ Add Row"))
 									];
 
 				CreateRowComboButton->SetOnGetMenuContent(FOnGetContent::CreateLambda([this, CreateRowComboButton]()
@@ -661,7 +661,7 @@ TSharedRef<SDockTab> FChooserTableEditor::SpawnTableTab( const FSpawnTabArgs& Ar
 	.ComboButtonStyle(FAppStyle::Get(), "SimpleComboButton")
 	.ButtonContent()
 	[
-		SNew(STextBlock).Text(LOCTEXT("Add", "+ Add Column"))
+		SNew(STextBlock).Text(LOCTEXT("AddColumn", "+ Add Column"))
 	];
 
 
@@ -907,7 +907,7 @@ TSharedRef<SWidget> CreateAssetWidget(UObject* Object)
 		{
 			if (DIAsset->Asset == nullptr)
 			{
-				return LOCTEXT("null", "[No Asset]");
+				return LOCTEXT("NoAsset", "[No Asset]");
 			}
 			else
 			{
@@ -936,7 +936,7 @@ TSharedRef<SWidget> CreateEvaluateChooserWidget(UObject* Object)
 		{
 			if (EvaluateChooser->Chooser == nullptr)
 			{
-				return LOCTEXT("null", "[No Chooser]");
+				return LOCTEXT("NoChooser", "[No Chooser]");
 			}
 			else
 			{
@@ -956,7 +956,7 @@ void ConvertToText_Asset(const UObject* Object, FText& OutText)
 	const UDataInterface_Object_Asset* AssetInterface = Cast<UDataInterface_Object_Asset>(Object);
 	if (AssetInterface->Asset == nullptr)
 	{
-		OutText = LOCTEXT("null", "[No Chooser]");
+		OutText = LOCTEXT("NoChooser", "[No Chooser]");
 	}
 	else
 	{
@@ -969,7 +969,7 @@ void ConvertToText_EvaluateChooser(const UObject* Object, FText& OutText)
 	const UDataInterface_EvaluateChooser* EvaluateChooser = Cast<UDataInterface_EvaluateChooser>(Object);
 	if (EvaluateChooser->Chooser == nullptr)
 	{
-		OutText = LOCTEXT("null", "[No Chooser]");
+		OutText = LOCTEXT("NoChooser", "[No Chooser]");
 	}
 	else
 	{
diff --git a/Engine/Plugins/Experimental/Dataflow/Source/DataflowEditor/Private/Dataflow/DataflowSchema.cpp b/Engine/Plugins/Experimental/Dataflow/Source/DataflowEditor/Private/Dataflow/DataflowSchema.cpp
--- a/Engine/Plugins/Experimental/Dataflow/Source/DataflowEditor/Private/Dataflow/DataflowSchema.cpp
+++ b/Engine/Plugins/Experimental/Dataflow/Source/DataflowEditor/Private/Dataflow/DataflowSchema.cpp
@@ -87,11 +87,11 @@ const FPinConnectionResponse UDataflowSchema::CanCreateConnection(const UEdGraph
 		}
 	}
 	TArray<FText> NoConnectionResponse = {
-		LOCTEXT("PinErrorSameNode", "Nope"), 
-		LOCTEXT("PinErrorSameNode", "Sorry :("),
-		LOCTEXT("PinErrorSameNode", "Not gonna work."),
-		LOCTEXT("PinErrorSameNode", "Still no!"),
-		LOCTEXT("PinErrorSameNode", "Try again?"),
+		LOCTEXT("PinErrorSameNode_Nope", "Nope"),
+		LOCTEXT("PinErrorSameNode_Sorry", "Sorry :("),
+		LOCTEXT("PinErrorSameNode_NotGonnaWork", "Not gonna work."),
+		LOCTEXT("PinErrorSameNode_StillNo", "Still no!"),
+		LOCTEXT("PinErrorSameNode_TryAgain", "Try again?"),
 	};
 	return FPinConnectionResponse(CONNECT_RESPONSE_DISALLOW, NoConnectionResponse[FMath::RandRange(0, NoConnectionResponse.Num()-1)]);
 }
diff --git a/Engine/Plugins/Experimental/Interchange/Runtime/Source/Messages/Private/Fbx/InterchangeFbxMessages.cpp b/Engine/Plugins/Experimental/Interchange/Runtime/Source/Messages/Private/Fbx/InterchangeFbxMessages.cpp
--- a/Engine/Plugins/Experimental/Interchange/Runtime/Source/Messages/Private/Fbx/InterchangeFbxMessages.cpp
+++ b/Engine/Plugins/Experimental/Interchange/Runtime/Source/Messages/Private/Fbx/InterchangeFbxMessages.cpp
@@ -50,10 +50,10 @@ FText UInterchangeResultTextureWarning_TextureFileDoNotExist::GetText() const
 
 	if (MaterialName.IsEmpty())
 	{
-		return FText::Format(LOCTEXT("TextureFileDoNotExist", "Cannot open texture file '{TextureName}'."), Args);
+		return FText::Format(LOCTEXT("TextureFileCannotBeOpened", "Cannot open texture file '{TextureName}'."), Args);
 	}
 	
-	return FText::Format(LOCTEXT("TextureFileDoNotExist", "Cannot open texture file '{TextureName}' when importing material '{MaterialName}'."), Args);
+	return FText::Format(LOCTEXT("CannotOpenTextureFileOnImport", "Cannot open texture file '{TextureName}' when importing material '{MaterialName}'."), Args);
 }
 
 
diff --git a/Engine/Plugins/Experimental/Interchange/Runtime/Source/Parsers/Fbx/Private/FbxAnimation.cpp b/Engine/Plugins/Experimental/Interchange/Runtime/Source/Parsers/Fbx/Private/FbxAnimation.cpp
--- a/Engine/Plugins/Experimental/Interchange/Runtime/Source/Parsers/Fbx/Private/FbxAnimation.cpp
+++ b/Engine/Plugins/Experimental/Interchange/Runtime/Source/Parsers/Fbx/Private/FbxAnimation.cpp
@@ -573,7 +573,7 @@ namespace UE::Interchange::Private
 		{
 			UInterchangeResultError_Generic* Message = Parser.AddMessage<UInterchangeResultError_Generic>();
 			Message->InterchangeKey = FFbxHelper::GetFbxNodeHierarchyName(Node);
-			Message->Text = LOCTEXT("FBXSceneNull", "Cannot fetch FBX animation transform payload because the FBX scene is null.");
+			Message->Text = LOCTEXT("FBXSceneNull_Animation", "Cannot fetch FBX animation transform payload because the FBX scene is null.");
 			return false;
 		}
 
diff --git a/Engine/Plugins/Experimental/Interchange/Runtime/Source/Parsers/Fbx/Private/FbxMesh.cpp b/Engine/Plugins/Experimental/Interchange/Runtime/Source/Parsers/Fbx/Private/FbxMesh.cpp
--- a/Engine/Plugins/Experimental/Interchange/Runtime/Source/Parsers/Fbx/Private/FbxMesh.cpp
+++ b/Engine/Plugins/Experimental/Interchange/Runtime/Source/Parsers/Fbx/Private/FbxMesh.cpp
@@ -1114,7 +1114,7 @@ bool FMeshPayloadContext::FetchPayloadToFile(FFbxParser& Parser, const FString&
 	{
 		UInterchangeResultError_Generic* Message = Parser.AddMessage<UInterchangeResultError_Generic>();
 		Message->InterchangeKey = FFbxHelper::GetMeshUniqueID(Mesh);
-		Message->Text = LOCTEXT("FBXSceneNull", "Cannot fetch FBX mesh payload because the FBX scene is null.");
+		Message->Text = LOCTEXT("FBXSceneNull_Mesh", "Cannot fetch FBX mesh payload because the FBX scene is null.");
 		return false;
 	}
 
diff --git a/Engine/Plugins/Experimental/MLDeformer/MLDeformerFramework/Source/MLDeformerFrameworkEditor/Private/MLDeformerEditorModel.cpp b/Engine/Plugins/Experimental/MLDeformer/MLDeformerFramework/Source/MLDeformerFrameworkEditor/Private/MLDeformerEditorModel.cpp
--- a/Engine/Plugins/Experimental/MLDeformer/MLDeformerFramework/Source/MLDeformerFrameworkEditor/Private/MLDeformerEditorModel.cpp
+++ b/Engine/Plugins/Experimental/MLDeformer/MLDeformerFramework/Source/MLDeformerFrameworkEditor/Private/MLDeformerEditorModel.cpp
@@ -752,7 +752,7 @@ namespace UE::MLDeformer
 				Result = FText::Format(LOCTEXT("TargetMeshMismatch", "Number of vertices in target mesh has changed from {0} to {1} vertices since this ML Deformer Asset was saved! {2}"),
 					InputInfo->GetNumTargetMeshVertices(),
 					Model->NumTargetMeshVerts,
-					IsTrained() ? LOCTEXT("BaseMeshMismatchNN", "Model needs to be retrained!") : FText());
+					IsTrained() ? LOCTEXT("BaseMeshMismatchModel", "Model needs to be retrained!") : FText());
 			}
 		}
 
diff --git a/Engine/Plugins/Experimental/PCG/Source/PCGEditor/Private/PCGEditor.cpp b/Engine/Plugins/Experimental/PCG/Source/PCGEditor/Private/PCGEditor.cpp
--- a/Engine/Plugins/Experimental/PCG/Source/PCGEditor/Private/PCGEditor.cpp
+++ b/Engine/Plugins/Experimental/PCG/Source/PCGEditor/Private/PCGEditor.cpp
@@ -262,7 +262,7 @@ void FPCGEditor::DeleteSelectedNodes()
 
 		bool bChanged = false;
 		{
-			const FScopedTransaction Transaction(*FPCGEditorCommon::ContextIdentifier, LOCTEXT("PCGEditorPaste", "PCG Editor: Delete"), nullptr);
+			const FScopedTransaction Transaction(*FPCGEditorCommon::ContextIdentifier, LOCTEXT("PCGEditorDeleteTransactionMessage", "PCG Editor: Delete"), nullptr);
 			PCGEditorGraph->Modify();
 
 			for (UObject* Object : GraphEditorWidget->GetSelectedNodes())
@@ -379,7 +379,7 @@ void FPCGEditor::PasteNodesHere(const FVector2D& Location)
 		return;
 	}
 
-	const FScopedTransaction Transaction(*FPCGEditorCommon::ContextIdentifier, LOCTEXT("PCGEditorPaste", "PCG Editor: Paste"), nullptr);
+	const FScopedTransaction Transaction(*FPCGEditorCommon::ContextIdentifier, LOCTEXT("PCGEditorPasteTransactionMessage", "PCG Editor: Paste"), nullptr);
 	PCGEditorGraph->Modify();
 
 	// Clear the selection set (newly pasted stuff will be selected)
diff --git a/Engine/Plugins/Experimental/PCG/Source/PCGEditor/Private/PCGEditorGraphNodeBase.cpp b/Engine/Plugins/Experimental/PCG/Source/PCGEditor/Private/PCGEditorGraphNodeBase.cpp
--- a/Engine/Plugins/Experimental/PCG/Source/PCGEditor/Private/PCGEditorGraphNodeBase.cpp
+++ b/Engine/Plugins/Experimental/PCG/Source/PCGEditor/Private/PCGEditorGraphNodeBase.cpp
@@ -400,7 +400,7 @@ void UPCGEditorGraphNodeBase::GetPinHoverText(const UEdGraphPin& Pin, FString& H
 	}
 	else
 	{
-		HoverTextOut = FText::Format(LOCTEXT("PinHoverToolTipFull", "Type: {0}\nSubtype: {1}"),
+		HoverTextOut = FText::Format(LOCTEXT("PinHoverToolTipPartial", "Type: {0}\nSubtype: {1}"),
 			DataTypeText,
 			DataSubtypeText).ToString();
 	}
diff --git a/Engine/Plugins/FX/Niagara/Source/Niagara/Private/NiagaraPlatformSet.cpp b/Engine/Plugins/FX/Niagara/Source/Niagara/Private/NiagaraPlatformSet.cpp
--- a/Engine/Plugins/FX/Niagara/Source/Niagara/Private/NiagaraPlatformSet.cpp
+++ b/Engine/Plugins/FX/Niagara/Source/Niagara/Private/NiagaraPlatformSet.cpp
@@ -641,7 +641,7 @@
 									}
 									else
 									{
-										OutDetails->ReasonsForDisabled.Emplace(FText::Format(LOCTEXT("PlatformSetDPDisabledFmt", "{0} is disabled."), FText::FromString(CurrProfile->GetName())));
+										OutDetails->ReasonsForDisabled.Emplace(FText::Format(LOCTEXT("PlatformSetDPDisabledFmt2", "{0} is disabled."), FText::FromString(CurrProfile->GetName())));
 									}
 								}
 								else
@@ -652,7 +652,7 @@
 									}
 									else
 									{
-										OutDetails->ReasonsForDisabled.Emplace(FText::Format(LOCTEXT("PlatformSetParentDPDisabledFmt", "{0} is Disabled."), FText::FromString(CurrProfile->GetName())));
+										OutDetails->ReasonsForDisabled.Emplace(FText::Format(LOCTEXT("PlatformSetParentDPDisabledFmt2", "{0} is Disabled."), FText::FromString(CurrProfile->GetName())));
 									}
 								}
 							}
@@ -741,7 +741,7 @@
 			}
 			if (bCanBeActiveDisabled)
 			{
-				FText ReasonTest = FText::Format(LOCTEXT("CVarConditionDisabledFmt", "{0} is disabled due to a condition on CVar {1}"), FText::FromString(Profile->GetName()), FText::FromName(CVarCondition.CVarName));
+				FText ReasonTest = FText::Format(LOCTEXT("CVarConditionDisabledFmt2", "{0} is disabled due to a condition on CVar {1}"), FText::FromString(Profile->GetName()), FText::FromName(CVarCondition.CVarName));
 				OutDetails->ReasonsForDisabled.Emplace(ReasonTest);
 			}
 		}
diff --git a/Engine/Plugins/FX/Niagara/Source/NiagaraEditor/Private/NiagaraNodeWithDynamicPins.cpp b/Engine/Plugins/FX/Niagara/Source/NiagaraEditor/Private/NiagaraNodeWithDynamicPins.cpp
--- a/Engine/Plugins/FX/Niagara/Source/NiagaraEditor/Private/NiagaraNodeWithDynamicPins.cpp
+++ b/Engine/Plugins/FX/Niagara/Source/NiagaraEditor/Private/NiagaraNodeWithDynamicPins.cpp
@@ -321,7 +321,7 @@ void UNiagaraNodeWithDynamicPins::GetNodeContextMenuActions(UToolMenu* Menu, UGr
 				Section.AddMenuEntry(
 				"MoveDynamicPinUp",
 				MoveUpLabel,
-				LOCTEXT("MoveDynamicPinToolTipUp", "Move this pin and any connections one slot up."),
+				LOCTEXT("MoveDynamicPinToolTipUpSlot", "Move this pin and any connections one slot up."),
 				FSlateIcon(),
 				FUIAction(FExecuteAction::CreateUObject(const_cast<UNiagaraNodeWithDynamicPins*>(this), &UNiagaraNodeWithDynamicPins::MoveDynamicPinFromMenu, const_cast<UEdGraphPin*>(Context->Pin), -1)));
 			}
diff --git a/Engine/Plugins/FX/Niagara/Source/NiagaraEditor/Private/ViewModels/Stack/NiagaraStackGraphUtilities.cpp b/Engine/Plugins/FX/Niagara/Source/NiagaraEditor/Private/ViewModels/Stack/NiagaraStackGraphUtilities.cpp
--- a/Engine/Plugins/FX/Niagara/Source/NiagaraEditor/Private/ViewModels/Stack/NiagaraStackGraphUtilities.cpp
+++ b/Engine/Plugins/FX/Niagara/Source/NiagaraEditor/Private/ViewModels/Stack/NiagaraStackGraphUtilities.cpp
@@ -486,7 +486,7 @@ void FNiagaraStackGraphUtilities::CheckForDeprecatedScriptVersion(UNiagaraNodeFu
 				ChangelistDescriptions.IsEmpty() ? FText() : FText::Format(LOCTEXT("DeprecatedVersionFormatChanges", "\n\nVersion change description:\n{0}"), ChangelistDescriptions));
 			UNiagaraStackEntry::FStackIssue UpgradeVersion(
                 EStackIssueSeverity::Warning,
-                FText::Format(LOCTEXT("DeprecatedVersionSummaryFormat", "Deprecated script version: {0}.{1} -> {2}.{3}"),
+                FText::Format(LOCTEXT("DeprecatedScriptVersionSummaryFormat", "Deprecated script version: {0}.{1} -> {2}.{3}"),
                 	FText::AsNumber(ReferencedVersion.MajorVersion), FText::AsNumber(ReferencedVersion.MinorVersion), FText::AsNumber(ExposedVersion.MajorVersion), FText::AsNumber(ExposedVersion.MinorVersion)),
                 LongDescription,
                 StackEditorDataKey,
@@ -574,10 +574,10 @@ void FNiagaraStackGraphUtilities::CheckForDeprecatedEmitterVersion(TSharedPtr<FN
 		if (FNiagaraAssetVersion ReferencedVersion = EmitterData->Version; ReferencedVersion.MajorVersion < ExposedVersion.MajorVersion)
 		{
 			TArray<UNiagaraStackEntry::FStackIssueFix> Fixes;
-			Fixes.Add(UNiagaraStackEntry::FStackIssueFix(LOCTEXT("UpgradeVersionFix", "Upgrade to newest parent version."), VersionUpgradeFix));
+			Fixes.Add(UNiagaraStackEntry::FStackIssueFix(LOCTEXT("UpgradeNewestParentFix", "Upgrade to newest parent version."), VersionUpgradeFix));
 			
 			FText ChangelistDescriptions = GetChangelistText(Emitter, ReferencedVersion, ExposedVersion);
-			FText LongDescription = FText::Format(LOCTEXT("DeprecatedVersionFormat", "This emitter has a newer parent version available.\nYou can upgrade now, but major version upgrades can sometimes come with breaking changes! So check that everything is still working as expected afterwards.{0}"),
+			FText LongDescription = FText::Format(LOCTEXT("EmitterDeprecatedVersionFormat", "This emitter has a newer parent version available.\nYou can upgrade now, but major version upgrades can sometimes come with breaking changes! So check that everything is still working as expected afterwards.{0}"),
 				ChangelistDescriptions.IsEmpty() ? FText() : FText::Format(LOCTEXT("DeprecatedVersionFormatChanges", "\n\nVersion change description:\n{0}"), ChangelistDescriptions));
 			UNiagaraStackEntry::FStackIssue UpgradeVersion(
                 EStackIssueSeverity::Warning,
@@ -604,7 +604,7 @@ void FNiagaraStackGraphUtilities::CheckForDeprecatedEmitterVersion(TSharedPtr<FN
 				UNiagaraStackEntry::FStackIssue UpgradeInfo(
 	                EStackIssueSeverity::Info,
 	                LOCTEXT("VersionUpgradeInfoSummary", "Version upgrade note"),
-	                FText::Format(LOCTEXT("VersionUpgradeFormatChanges", "The version of this emitter was recently upgraded; here is a list of changes:\n{0}"), ChangelistDescriptions),
+	                FText::Format(LOCTEXT("EmitterVersionUpgradeFormatChanges", "The version of this emitter was recently upgraded; here is a list of changes:\n{0}"), ChangelistDescriptions),
 	                StackEditorDataKey,
 	                true);
 				OutIssues.Add(UpgradeInfo);
diff --git a/Engine/Plugins/FX/Niagara/Source/NiagaraEditor/Private/ViewModels/Stack/NiagaraStackScriptItemGroup.cpp b/Engine/Plugins/FX/Niagara/Source/NiagaraEditor/Private/ViewModels/Stack/NiagaraStackScriptItemGroup.cpp
--- a/Engine/Plugins/FX/Niagara/Source/NiagaraEditor/Private/ViewModels/Stack/NiagaraStackScriptItemGroup.cpp
+++ b/Engine/Plugins/FX/Niagara/Source/NiagaraEditor/Private/ViewModels/Stack/NiagaraStackScriptItemGroup.cpp
@@ -996,7 +996,7 @@ TOptional<UNiagaraStackEntry::FDropRequestResponse> UNiagaraStackScriptItemGroup
 	if (DropRequest.DropOptions != UNiagaraStackEntry::EDropOptions::Overview)
 	{
 		// Only allow dropping in the overview stacks.
-		return FDropRequestResponse(TOptional<EItemDropZone>(), LOCTEXT("AssetCantDropOnStack", "Scripts can only be dropped into the overview."));
+		return FDropRequestResponse(TOptional<EItemDropZone>(), LOCTEXT("ScriptCantDropOnStack", "Scripts can only be dropped into the overview."));
 	}
 	UNiagaraScript* Script = DragDropOp->Script.Get();
 	if (Script)
diff --git a/Engine/Plugins/FX/Niagara/Source/NiagaraEditor/Private/Widgets/SNiagaraBakerWidget.cpp b/Engine/Plugins/FX/Niagara/Source/NiagaraEditor/Private/Widgets/SNiagaraBakerWidget.cpp
--- a/Engine/Plugins/FX/Niagara/Source/NiagaraEditor/Private/Widgets/SNiagaraBakerWidget.cpp
+++ b/Engine/Plugins/FX/Niagara/Source/NiagaraEditor/Private/Widgets/SNiagaraBakerWidget.cpp
@@ -653,7 +653,7 @@ TSharedRef<SWidget> SNiagaraBakerWidget::MakeViewOptionsMenu()
 				FMakeRotatorBox::FGetter::CreateSP(ViewModel, &FNiagaraBakerViewModel::GetCurrentCameraRotation),
 				FMakeRotatorBox::FSetter::CreateSP(ViewModel, &FNiagaraBakerViewModel::SetCurrentCameraRotation)
 			),
-			LOCTEXT("CameraLocation", "Camera Rotation")
+			LOCTEXT("CameraRotation", "Camera Rotation")
 		);
 
 		if ( ViewModel->IsCurrentCameraPerspective() )
diff --git a/Engine/Plugins/FX/Niagara/Source/NiagaraEditorWidgets/Private/SNiagaraOverviewGraph.cpp b/Engine/Plugins/FX/Niagara/Source/NiagaraEditorWidgets/Private/SNiagaraOverviewGraph.cpp
--- a/Engine/Plugins/FX/Niagara/Source/NiagaraEditorWidgets/Private/SNiagaraOverviewGraph.cpp
+++ b/Engine/Plugins/FX/Niagara/Source/NiagaraEditorWidgets/Private/SNiagaraOverviewGraph.cpp
@@ -270,12 +270,12 @@ FActionMenuContent SNiagaraOverviewGraph::OnCreateGraphActionMenu(UEdGraph* InGr
 			MenuBuilder.AddMenuEntry(FNiagaraEditorCommands::Get().OpenAddEmitterMenu);
 			MenuBuilder.AddMenuEntry(
 				LOCTEXT("EmptyEmitterLabel", "Add empty emitter"),
-				LOCTEXT("CommentsToolTip", "Adds an empty emitter without any modules or renderers."),
+				LOCTEXT("AddEmitterToolTip", "Adds an empty emitter without any modules or renderers."),
 				FSlateIcon(),
 				FExecuteAction::CreateSP(this, &SNiagaraOverviewGraph::OnCreateEmptyEmitter));
 			MenuBuilder.AddMenuEntry(
 				LOCTEXT("CommentsLabel", "Add Comment"),
-				LOCTEXT("CommentsToolTip", "Add a comment box"),
+				LOCTEXT("AddCommentBoxToolTip", "Add a comment box"),
 				FSlateIcon(),
 				FExecuteAction::CreateSP(this, &SNiagaraOverviewGraph::OnCreateComment));
 
diff --git a/Engine/Plugins/Media/ImgMedia/Source/ImgMediaEditor/Private/Widgets/SImgMediaProcessImages.cpp b/Engine/Plugins/Media/ImgMedia/Source/ImgMediaEditor/Private/Widgets/SImgMediaProcessImages.cpp
--- a/Engine/Plugins/Media/ImgMedia/Source/ImgMediaEditor/Private/Widgets/SImgMediaProcessImages.cpp
+++ b/Engine/Plugins/Media/ImgMedia/Source/ImgMediaEditor/Private/Widgets/SImgMediaProcessImages.cpp
@@ -76,7 +76,7 @@ void SImgMediaProcessImages::Construct(const FArguments& InArgs)
 						SAssignNew(StartButton, SButton)
 							.OnClicked(this, &SImgMediaProcessImages::OnProcessImagesClicked)
 							.Text(LOCTEXT("StartProcessImages", "Process Images"))
-							.ToolTipText(LOCTEXT("ProcesssImagesButtonToolTip", "Start processing images."))
+							.ToolTipText(LOCTEXT("StartProcesssImagesButtonToolTip", "Start processing images."))
 					]
 
 				+ SHorizontalBox::Slot()
@@ -85,7 +85,7 @@ void SImgMediaProcessImages::Construct(const FArguments& InArgs)
 						SAssignNew(CancelButton, SButton)
 							.OnClicked(this, &SImgMediaProcessImages::OnCancelClicked)
 							.Text(LOCTEXT("CancelProcessImages", "Cancel"))
-							.ToolTipText(LOCTEXT("ProcesssImagesButtonToolTip", "Cancel processing images."))
+							.ToolTipText(LOCTEXT("CancelProcesssImagesButtonToolTip", "Cancel processing images."))
 					]
 			]
 	];
@@ -933,7 +933,7 @@ void SImgMediaProcessImages::HandleProcessing()
 				if (ConfirmNotification.IsValid())
 				{
 					ConfirmNotification->SetText(
-						FText::Format(LOCTEXT("ImgMediaCompleted", "ImgMedia Completed {0}"),
+						FText::Format(LOCTEXT("ImgMediaCompleted2", "ImgMedia Completed {0}"),
 							FText::AsNumber(CurrentFrameIndex)));
 				}
 
diff --git a/Engine/Plugins/Runtime/MeshModelingToolset/Source/MeshModelingTools/Private/AddPrimitiveTool.cpp b/Engine/Plugins/Runtime/MeshModelingToolset/Source/MeshModelingTools/Private/AddPrimitiveTool.cpp
--- a/Engine/Plugins/Runtime/MeshModelingToolset/Source/MeshModelingTools/Private/AddPrimitiveTool.cpp
+++ b/Engine/Plugins/Runtime/MeshModelingToolset/Source/MeshModelingTools/Private/AddPrimitiveTool.cpp
@@ -170,7 +170,7 @@ void UAddPrimitiveTool::SetState(EState NewState)
 		Gizmo->ReinitializeGizmoTransform(PreviewMesh->GetTransform());
 
 		GetToolManager()->DisplayMessage(
-			LOCTEXT("OnStartAddPrimitiveTool", "Alter shape settings in the detail panel or modify placement with gizmo, then accept the tool."),
+			LOCTEXT("OnStartAddPrimitiveTool2", "Alter shape settings in the detail panel or modify placement with gizmo, then accept the tool."),
 			EToolMessageLevel::UserNotification);
 	}
 }
diff --git a/Engine/Plugins/Runtime/Metasound/Source/MetasoundEditor/Private/MetasoundAssetTypeActions.cpp b/Engine/Plugins/Runtime/Metasound/Source/MetasoundEditor/Private/MetasoundAssetTypeActions.cpp
--- a/Engine/Plugins/Runtime/Metasound/Source/MetasoundEditor/Private/MetasoundAssetTypeActions.cpp
+++ b/Engine/Plugins/Runtime/Metasound/Source/MetasoundEditor/Private/MetasoundAssetTypeActions.cpp
@@ -316,11 +316,11 @@ namespace Metasound
 				FText Format;
 				if (IsPlaying(Path))
 				{
-					Format = LOCTEXT("PreviewMetaSoundFromIconToolTip", "Stop Previewing {0}");
+					Format = LOCTEXT("StopPreviewMetaSoundFromIconToolTip", "Stop Previewing {0}");
 				}
 				else
 				{
-					Format = LOCTEXT("PreviewMetaSoundFromIconToolTip", "Preview {0}");
+					Format = LOCTEXT("PreviewMetaSoundFromIconToolTip_Editor", "Preview {0}");
 				}
 
 				FName TypeName = UMetaSoundSource::StaticClass()->GetFName();
diff --git a/Engine/Plugins/Runtime/Metasound/Source/MetasoundEditor/Private/SMetasoundPinValueInspector.h b/Engine/Plugins/Runtime/Metasound/Source/MetasoundEditor/Private/SMetasoundPinValueInspector.h
--- a/Engine/Plugins/Runtime/Metasound/Source/MetasoundEditor/Private/SMetasoundPinValueInspector.h
+++ b/Engine/Plugins/Runtime/Metasound/Source/MetasoundEditor/Private/SMetasoundPinValueInspector.h
@@ -118,4 +118,4 @@ namespace Metasound
 		};
 	} // namespace Editor
 } // namespace Metasound
-#undef LOCTEXT_NAMESPACE
+
diff --git a/Engine/Plugins/Runtime/ModelViewViewModel/Source/ModelViewViewModel/Private/View/MVVMView.cpp b/Engine/Plugins/Runtime/ModelViewViewModel/Source/ModelViewViewModel/Private/View/MVVMView.cpp
--- a/Engine/Plugins/Runtime/ModelViewViewModel/Source/ModelViewViewModel/Private/View/MVVMView.cpp
+++ b/Engine/Plugins/Runtime/ModelViewViewModel/Source/ModelViewViewModel/Private/View/MVVMView.cpp
@@ -349,13 +349,13 @@ bool UMVVMView::RegisterLibraryBinding(const FMVVMViewClass_CompiledBinding& Bin
 	if (ClassExtension->GetBindingLibrary().EvaluateFieldPath(GetUserWidget(), Binding.GetBinding().GetSourceFieldPath()).HasError())
 	{
 		UE::MVVM::FMessageLog Log(GetUserWidget());
-		Log.Warning(FText::Format(LOCTEXT("RegisterBindingInvalidDestination", "Can't register binding '{0}'. The source was not evaluated."), FText::FromString(Binding.ToString())));
+		Log.Warning(FText::Format(LOCTEXT("RegisterBindingSourceNotEvaluated", "Can't register binding '{0}'. The source was not evaluated."), FText::FromString(Binding.ToString())));
 		return false;
 	}
 	if (ClassExtension->GetBindingLibrary().EvaluateFieldPath(GetUserWidget(), Binding.GetBinding().GetDestinationFieldPath()).HasError())
 	{
 		UE::MVVM::FMessageLog Log(GetUserWidget());
-		Log.Warning(FText::Format(LOCTEXT("RegisterBindingInvalidDestination", "Can't register binding '{0}'. The destination was not evaluated."), FText::FromString(Binding.ToString())));
+		Log.Warning(FText::Format(LOCTEXT("RegisterBindingInvalidDestination2", "Can't register binding '{0}'. The destination was not evaluated."), FText::FromString(Binding.ToString())));
 		return false;
 	}
 
@@ -443,4 +443,4 @@ TScriptInterface<INotifyFieldValueChanged> UMVVMView::FindSource(const FMVVMView
 }
 
 
-#undef LOCTEXT_NAMESPACE
\ No newline at end of file
+#undef LOCTEXT_NAMESPACE
diff --git a/Engine/Plugins/Runtime/ModelViewViewModel/Source/ModelViewViewModelEditor/Private/Widgets/SMVVMViewBindingListView.cpp b/Engine/Plugins/Runtime/ModelViewViewModel/Source/ModelViewViewModelEditor/Private/Widgets/SMVVMViewBindingListView.cpp
--- a/Engine/Plugins/Runtime/ModelViewViewModel/Source/ModelViewViewModelEditor/Private/Widgets/SMVVMViewBindingListView.cpp
+++ b/Engine/Plugins/Runtime/ModelViewViewModel/Source/ModelViewViewModelEditor/Private/Widgets/SMVVMViewBindingListView.cpp
@@ -819,7 +819,7 @@ void SMVVMViewBindingListView::Construct(const FArguments& InArgs, TSharedPtr<SM
 			.DefaultLabel(LOCTEXT("Widget", "Widget"))
 			+ SHeaderRow::Column(SMVVMViewBindingListEntryRow::WidgetPropertyColumnName)
 			.FillWidth(0.125f)
-			.DefaultLabel(LOCTEXT("Widget", "Widget Property"))
+			.DefaultLabel(LOCTEXT("WidgetProperty", "Widget Property"))
 			+ SHeaderRow::Column(SMVVMViewBindingListEntryRow::UpdateColumnName)
 			.FillWidth(0.05f)
 			.DefaultLabel(LOCTEXT("Update", "Update"))
diff --git a/Engine/Plugins/Runtime/ModelViewViewModel/Source/ModelViewViewModelEditor/Private/Widgets/SMVVMViewModelContextListWidget.cpp b/Engine/Plugins/Runtime/ModelViewViewModel/Source/ModelViewViewModelEditor/Private/Widgets/SMVVMViewModelContextListWidget.cpp
--- a/Engine/Plugins/Runtime/ModelViewViewModel/Source/ModelViewViewModelEditor/Private/Widgets/SMVVMViewModelContextListWidget.cpp
+++ b/Engine/Plugins/Runtime/ModelViewViewModel/Source/ModelViewViewModelEditor/Private/Widgets/SMVVMViewModelContextListWidget.cpp
@@ -711,7 +711,7 @@ void SMVVMViewModelContextListWidget::Construct(const FArguments& InArgs)
 				SNew(SHeaderRow)
 				+ SHeaderRow::Column(SMVVMManageViewModelsListEntryRow::RemoveButtonColumnName)
 				.FixedWidth(30.0f)
-				.DefaultLabel(LOCTEXT("", "")) 
+				.DefaultLabel(FText::GetEmpty()) 
 				+ SHeaderRow::Column(SMVVMManageViewModelsListEntryRow::ClassColumnName)
 				.DefaultLabel(LOCTEXT("Class", "Class"))
 				+ SHeaderRow::Column(SMVVMManageViewModelsListEntryRow::NameColumnName)
diff --git a/Engine/Plugins/Runtime/nDisplay/Source/DisplayClusterLightCardEditor/Private/SDisplayClusterLightCardList.cpp b/Engine/Plugins/Runtime/nDisplay/Source/DisplayClusterLightCardEditor/Private/SDisplayClusterLightCardList.cpp
--- a/Engine/Plugins/Runtime/nDisplay/Source/DisplayClusterLightCardEditor/Private/SDisplayClusterLightCardList.cpp
+++ b/Engine/Plugins/Runtime/nDisplay/Source/DisplayClusterLightCardEditor/Private/SDisplayClusterLightCardList.cpp
@@ -501,7 +501,7 @@
 	GEditor->EditorAddModalWindow(PickerWindow.ToSharedRef());
 	if (bFinished && SelectedActorPtr.IsValid())
 	{
-		FScopedTransaction Transaction(LOCTEXT("AddExistingLightCard", "Add Existing Light Card"));
+		FScopedTransaction Transaction(LOCTEXT("AddExistingLightCardTransactionMessage", "Add Existing Light Card"));
 		AddLightCardToActor(SelectedActorPtr.Get());
 	}
 
diff --git a/Engine/Plugins/VirtualProduction/DMX/DMXEngine/Source/DMXRuntime/Private/MVR/DMXMVRUnzip.cpp b/Engine/Plugins/VirtualProduction/DMX/DMXEngine/Source/DMXRuntime/Private/MVR/DMXMVRUnzip.cpp
--- a/Engine/Plugins/VirtualProduction/DMX/DMXEngine/Source/DMXRuntime/Private/MVR/DMXMVRUnzip.cpp
+++ b/Engine/Plugins/VirtualProduction/DMX/DMXEngine/Source/DMXRuntime/Private/MVR/DMXMVRUnzip.cpp
@@ -243,4 +243,4 @@ bool FDMXMVRUnzip::InitializeFromFileInternal(const FString& FilePathAndName)
 	return InitializeFromDataInternal(NewData.GetData(), NewData.Num());
 }
 
-#undef LOCTEXT_NAMESPACE
+
diff --git a/Engine/Plugins/VirtualProduction/RemoteControl/Source/RemoteControlUI/Private/UI/SRemoteControlPanel.cpp b/Engine/Plugins/VirtualProduction/RemoteControl/Source/RemoteControlUI/Private/UI/SRemoteControlPanel.cpp
--- a/Engine/Plugins/VirtualProduction/RemoteControl/Source/RemoteControlUI/Private/UI/SRemoteControlPanel.cpp
+++ b/Engine/Plugins/VirtualProduction/RemoteControl/Source/RemoteControlUI/Private/UI/SRemoteControlPanel.cpp
@@ -257,7 +257,7 @@ void SRemoteControlPanel::Construct(const FArguments& InArgs, URemoteControlPres
 		.Content()
 		[
 			SNew(STextBlock)
-			.Text(LOCTEXT("EditModeLabel", "Edit"))
+			.Text(LOCTEXT("EditLabel", "Edit"))
 		]
 	);
 
@@ -508,7 +508,7 @@ void SRemoteControlPanel::Construct(const FArguments& InArgs, URemoteControlPres
 				.VAlign(VAlign_Center)
 				[
 					SNew(STextBlock)
-					.Text(LOCTEXT("EditModeLabel", "Edit Mode : "))
+					.Text(LOCTEXT("EditModeLabel_Editor", "Edit Mode : "))
 				]
 				+ SHorizontalBox::Slot()
 				.AutoWidth()
diff --git a/Engine/Plugins/Web/HttpBlueprint/Source/HttpBlueprintGraph/Private/K2Node_HttpRequest.cpp b/Engine/Plugins/Web/HttpBlueprint/Source/HttpBlueprintGraph/Private/K2Node_HttpRequest.cpp
--- a/Engine/Plugins/Web/HttpBlueprint/Source/HttpBlueprintGraph/Private/K2Node_HttpRequest.cpp
+++ b/Engine/Plugins/Web/HttpBlueprint/Source/HttpBlueprintGraph/Private/K2Node_HttpRequest.cpp
@@ -169,7 +169,7 @@ void UK2Node_HttpRequest::ExpandNode(FKismetCompilerContext& CompilerContext, UE
 	else if (GetBodyPin())
 	{
 		CompilerContext.MessageLog.Error(
-			*LOCTEXT("Error", "Node @@ Attempted to supply a input body pin that was not a string or struct.").
+			*LOCTEXT("Error_InputPinNotStringOrStruct", "Node @@ Attempted to supply a input body pin that was not a string or struct.").
 			ToString(), this);
 		BreakAllNodeLinks();
 	}
@@ -255,7 +255,7 @@ void UK2Node_HttpRequest::ExpandNode(FKismetCompilerContext& CompilerContext, UE
 	else
 	{
 		CompilerContext.MessageLog.Error(
-			*LOCTEXT("Error", "Node @@ Attempted to deserialize the Result Body into a @@. The supported types are String and Struct").
+			*LOCTEXT("Error_DeserializeIntoUnsupportedType", "Node @@ Attempted to deserialize the Result Body into a @@. The supported types are String and Struct").
 			ToString(), this, *CachedOutputBodyPinType.PinCategory.ToString());
 		BreakAllNodeLinks();
 	}
@@ -306,7 +306,7 @@ bool UK2Node_HttpRequest::IsConnectionDisallowed(
 			return false;
 		}
 
-		OutReason = LOCTEXT("HttpRequest_InvalidArgumentType", "Output Body may only be a string or struct").ToString();
+		OutReason = LOCTEXT("HttpRequest_InvalidArgumentType_Output", "Output Body may only be a string or struct").ToString();
 		return true;
 	}
 	
@@ -322,7 +322,7 @@ bool UK2Node_HttpRequest::IsConnectionDisallowed(
 
 		if (!bIsValidType)
 		{
-			OutReason = LOCTEXT("HttpRequest_InvalidArgumentType", "Input Body may only be a string or struct").ToString();
+			OutReason = LOCTEXT("HttpRequest_InvalidArgumentType_Input", "Input Body may only be a string or struct").ToString();
 			return true;
 		}
 	}
diff --git a/Engine/Plugins/Web/WebAPI/Source/WebAPIEditor/Private/Details/ViewModels/WebAPIEnumViewModel.cpp b/Engine/Plugins/Web/WebAPI/Source/WebAPIEditor/Private/Details/ViewModels/WebAPIEnumViewModel.cpp
--- a/Engine/Plugins/Web/WebAPI/Source/WebAPIEditor/Private/Details/ViewModels/WebAPIEnumViewModel.cpp
+++ b/Engine/Plugins/Web/WebAPI/Source/WebAPIEditor/Private/Details/ViewModels/WebAPIEnumViewModel.cpp
@@ -1,4 +1,4 @@
-﻿// Copyright Epic Games, Inc. All Rights Reserved.
+// Copyright Epic Games, Inc. All Rights Reserved.
 
 #include "WebAPIEnumViewModel.h"
 
@@ -83,7 +83,7 @@ void FWebAPIEnumViewModel::Initialize()
 
 	if(Enum->Name.HasTypeInfo() && Enum->Name.TypeInfo->bIsNested)
 	{
-		TooltipArgs.Emplace(LOCTEXT("EnumTooltip_Header", "Enum (Nested)").ToUpper());
+		TooltipArgs.Emplace(LOCTEXT("NestedEnumTooltip_Header", "Enum (Nested)").ToUpper());
 		TooltipArgs.Emplace(FText::Format(LOCTEXT("EnumTooltip_MemberName", "MemberName: {0}"), FText::FromString(Enum->Name.TypeInfo->ToMemberName())));
 	}
 	else
diff --git a/Engine/Plugins/Web/WebAPI/Source/WebAPIEditor/Private/Details/ViewModels/WebAPIModelViewModel.cpp b/Engine/Plugins/Web/WebAPI/Source/WebAPIEditor/Private/Details/ViewModels/WebAPIModelViewModel.cpp
--- a/Engine/Plugins/Web/WebAPI/Source/WebAPIEditor/Private/Details/ViewModels/WebAPIModelViewModel.cpp
+++ b/Engine/Plugins/Web/WebAPI/Source/WebAPIEditor/Private/Details/ViewModels/WebAPIModelViewModel.cpp
@@ -1,4 +1,4 @@
-﻿// Copyright Epic Games, Inc. All Rights Reserved.
+// Copyright Epic Games, Inc. All Rights Reserved.
 
 #include "WebAPIModelViewModel.h"
 
@@ -181,7 +181,7 @@ void FWebAPIModelViewModel::Initialize()
 	TArray<FText> TooltipArgs;
 	if(Model->Name.HasTypeInfo() && Model->Name.TypeInfo.LoadSynchronous()->bIsNested)
 	{
-		TooltipArgs.Emplace(LOCTEXT("ModelTooltip_Header", "Model (Nested)").ToUpper());
+		TooltipArgs.Emplace(LOCTEXT("NestedModelTooltip_Header", "Model (Nested)").ToUpper());
 	}
 	else
 	{
diff --git a/Engine/Plugins/Web/WebAPI/Source/WebAPIEditor/Private/Details/ViewModels/WebAPIOperationResponseViewModel.cpp b/Engine/Plugins/Web/WebAPI/Source/WebAPIEditor/Private/Details/ViewModels/WebAPIOperationResponseViewModel.cpp
--- a/Engine/Plugins/Web/WebAPI/Source/WebAPIEditor/Private/Details/ViewModels/WebAPIOperationResponseViewModel.cpp
+++ b/Engine/Plugins/Web/WebAPI/Source/WebAPIEditor/Private/Details/ViewModels/WebAPIOperationResponseViewModel.cpp
@@ -1,4 +1,4 @@
-﻿// Copyright Epic Games, Inc. All Rights Reserved.
+// Copyright Epic Games, Inc. All Rights Reserved.
 
 #include "WebAPIOperationResponseViewModel.h"
 
@@ -52,7 +52,7 @@ void FWebAPIOperationResponseViewModel::Initialize()
 	}
 	else
 	{
-		CachedLabel = LOCTEXT("ResponseLabel", "Response Code: Default");
+		CachedLabel = LOCTEXT("ResponseLabel_Default", "Response Code: Default");
 	}
 
 	// Check if response is empty, but has description
diff --git a/Engine/Plugins/Web/WebAPI/Source/WebAPIEditor/Private/Dom/WebAPIModel.cpp b/Engine/Plugins/Web/WebAPI/Source/WebAPIEditor/Private/Dom/WebAPIModel.cpp
--- a/Engine/Plugins/Web/WebAPI/Source/WebAPIEditor/Private/Dom/WebAPIModel.cpp
+++ b/Engine/Plugins/Web/WebAPI/Source/WebAPIEditor/Private/Dom/WebAPIModel.cpp
@@ -1,4 +1,4 @@
-﻿// Copyright Epic Games, Inc. All Rights Reserved.
+// Copyright Epic Games, Inc. All Rights Reserved.
 
 #include "Dom/WebAPIModel.h"
 
@@ -135,7 +135,7 @@ EDataValidationResult UWebAPIModel::IsDataValid(TArray<FText>& ValidationErrors)
 
 	if(Name.ToString(true).IsEmpty())
 	{
-		ValidationErrors.Add(LOCTEXT("Missing_Name", "Model missing name"));
+		ValidationErrors.Add(LOCTEXT("Missing_Model_Name", "Model missing name"));
 		ValidationResult = EDataValidationResult::Invalid;
 	}
 
diff --git a/Engine/Plugins/Web/WebAPI/Source/WebAPILiquidJS/Private/WebAPILiquidJSProcess.cpp b/Engine/Plugins/Web/WebAPI/Source/WebAPILiquidJS/Private/WebAPILiquidJSProcess.cpp
--- a/Engine/Plugins/Web/WebAPI/Source/WebAPILiquidJS/Private/WebAPILiquidJSProcess.cpp
+++ b/Engine/Plugins/Web/WebAPI/Source/WebAPILiquidJS/Private/WebAPILiquidJSProcess.cpp
@@ -1,4 +1,4 @@
-﻿// Copyright Epic Games, Inc. All Rights Reserved.
+// Copyright Epic Games, Inc. All Rights Reserved.
 
 #include "WebAPILiquidJSProcess.h"
 
@@ -260,7 +260,7 @@ uint32 FWebAPILiquidJSProcess::Run()
 	{
 		TaskNotification->SetComplete(
 			ErrorTitle,
-			LOCTEXT("WebAPILiquidJS_LaunchFailed", "WebApp exited"),
+			LOCTEXT("WebAPILiquidJS_Exited", "WebApp exited"),
 			false
 		);
 	}
@@ -277,4 +277,4 @@ uint32 FWebAPILiquidJSProcess::Run()
 	return 0;
 }
 
-#undef LOCTEXT_NAMESPACE
\ No newline at end of file
+#undef LOCTEXT_NAMESPACE
diff --git a/Engine/Source/Developer/ProjectLauncher/Private/Widgets/Package/SProjectLauncherPackagingSettings.cpp b/Engine/Source/Developer/ProjectLauncher/Private/Widgets/Package/SProjectLauncherPackagingSettings.cpp
--- a/Engine/Source/Developer/ProjectLauncher/Private/Widgets/Package/SProjectLauncherPackagingSettings.cpp
+++ b/Engine/Source/Developer/ProjectLauncher/Private/Widgets/Package/SProjectLauncherPackagingSettings.cpp
@@ -185,8 +185,8 @@ void SProjectLauncherPackagingSettings::Construct(const FArguments& InArgs, cons
 										.IsEnabled(this, &SProjectLauncherPackagingSettings::IsReferenceBlockDbEditable)
 										.OnTextCommitted(this, &SProjectLauncherPackagingSettings::OnRefBlockDbFileNameTextCommitted)
 										.OnTextChanged(this, &SProjectLauncherPackagingSettings::OnRefBlockDbFileNameTextChanged)
-										.ToolTipText(LOCTEXT("RefBlockDbFileNameToolTip", "The path to a global.utoc in a directory of previously released iostore containers."))
-										.HintText(LOCTEXT("RefBlockDbFileNameHint", "<path/to/global.utoc>"))
+										.ToolTipText(LOCTEXT("RefBlockDbUtocFileNameToolTip", "The path to a global.utoc in a directory of previously released iostore containers."))
+										.HintText(LOCTEXT("RefBlockDbUtocFileNameHint", "<path/to/global.utoc>"))
 									]
 
 								+ SHorizontalBox::Slot()
@@ -199,7 +199,7 @@ void SProjectLauncherPackagingSettings::Construct(const FArguments& InArgs, cons
 										.ContentPadding(FMargin(6.0, 2.0))
 										.IsEnabled(this, &SProjectLauncherPackagingSettings::IsReferenceBlockDbEditable)
 										.Text(LOCTEXT("BrowseButtonText", "Browse..."))
-										.ToolTipText(LOCTEXT("BrowseForRefDbButtonToolTip", "Browse for the global.utoc file from a previously staged/released game."))
+										.ToolTipText(LOCTEXT("BrowseForRefDbUtocFilenameButtonToolTip", "Browse for the global.utoc file from a previously staged/released game."))
 										.OnClicked(this, &SProjectLauncherPackagingSettings::HandleRefBlockDbBrowseButtonClicked)
 									]
 							] // end reference container line
@@ -219,8 +219,8 @@ void SProjectLauncherPackagingSettings::Construct(const FArguments& InArgs, cons
 									.IsEnabled(this, &SProjectLauncherPackagingSettings::IsReferenceBlockDbEditable)
 									.OnTextCommitted(this, &SProjectLauncherPackagingSettings::OnRefBlockCryptoFileNameTextCommitted)
 									.OnTextChanged(this, &SProjectLauncherPackagingSettings::OnRefBlockCryptoFileNameTextChanged)
-									.ToolTipText(LOCTEXT("RefBlockDbFileNameToolTip", "The path to a crypto.json to decrypt the reference block containers, if needed."))
-									.HintText(LOCTEXT("RefBlockDbFileNameHint", "<path/to/crypto.json, optional>"))
+									.ToolTipText(LOCTEXT("RefBlockDbCryptoFileNameToolTip", "The path to a crypto.json to decrypt the reference block containers, if needed."))
+									.HintText(LOCTEXT("RefBlockDbCryptoFileNameHint", "<path/to/crypto.json, optional>"))
 								]
 
 								+ SHorizontalBox::Slot()
@@ -233,7 +233,7 @@ void SProjectLauncherPackagingSettings::Construct(const FArguments& InArgs, cons
 										.ContentPadding(FMargin(6.0, 2.0))
 										.IsEnabled(this, &SProjectLauncherPackagingSettings::IsReferenceBlockDbEditable)
 										.Text(LOCTEXT("BrowseButtonText", "Browse..."))
-										.ToolTipText(LOCTEXT("BrowseForRefDbButtonToolTip", "Browse for the crypto.json file to decrypt the reference containers, if needed."))
+										.ToolTipText(LOCTEXT("BrowseForRefDbCryptoFilenameButtonToolTip", "Browse for the crypto.json file to decrypt the reference containers, if needed."))
 										.OnClicked(this, &SProjectLauncherPackagingSettings::HandleRefBlockCryptoBrowseButtonClicked)
 									]
 							] // end reference container crypto keys line
diff --git a/Engine/Source/Developer/TraceInsights/Private/Insights/Widgets/SStartPageWindow.cpp b/Engine/Source/Developer/TraceInsights/Private/Insights/Widgets/SStartPageWindow.cpp
--- a/Engine/Source/Developer/TraceInsights/Private/Insights/Widgets/SStartPageWindow.cpp
+++ b/Engine/Source/Developer/TraceInsights/Private/Insights/Widgets/SStartPageWindow.cpp
@@ -2446,7 +2446,7 @@ FTraceFilterByBuildConfig::FTraceFilterByBuildConfig()
 
 FTraceFilterByBuildTarget::FTraceFilterByBuildTarget()
 {
-	ToggleAllActionTooltip = LOCTEXT("FilterByBuildConfig_ToggleAll_Tooltip", "Shows or hides traces for all build targets.");
+	ToggleAllActionTooltip = LOCTEXT("FilterByBuildTarget_ToggleAll_Tooltip", "Shows or hides traces for all build targets.");
 }
 
 ////////////////////////////////////////////////////////////////////////////////////////////////////
diff --git a/Engine/Source/Editor/BlueprintGraph/Private/EdGraphSchema_K2.cpp b/Engine/Source/Editor/BlueprintGraph/Private/EdGraphSchema_K2.cpp
--- a/Engine/Source/Editor/BlueprintGraph/Private/EdGraphSchema_K2.cpp
+++ b/Engine/Source/Editor/BlueprintGraph/Private/EdGraphSchema_K2.cpp
@@ -3710,18 +3710,18 @@ FText UEdGraphSchema_K2::GetCategoryText(FName Category, FName SubCategory, bool
 		{
 		case EBlueprintRealDisplayMode::Real:
 			CategoryDescriptions[PC_Real] = LOCTEXT("RealCategory_DisplayAsReal", "Real");
-			CategoryDescriptions[PC_Float] = LOCTEXT("RealCategory_DisplayAsReal", "Real (single-precision)");
-			CategoryDescriptions[PC_Double] = LOCTEXT("RealCategory_DisplayAsReal", "Real (double-precision)");
+			CategoryDescriptions[PC_Float] = LOCTEXT("RealCategory_DisplayAsReal_SinglePrecision", "Real (single-precision)");
+			CategoryDescriptions[PC_Double] = LOCTEXT("RealCategory_DisplayAsReal_DoublePrecision", "Real (double-precision)");
 			break;
 		case EBlueprintRealDisplayMode::Float:
 			CategoryDescriptions[PC_Real] = LOCTEXT("RealCategory_DisplayAsFloat", "Float");
-			CategoryDescriptions[PC_Float] = LOCTEXT("RealCategory_DisplayAsFloat", "Float (single-precision)");
-			CategoryDescriptions[PC_Double] = LOCTEXT("RealCategory_DisplayAsFloat", "Float (double-precision)");
+			CategoryDescriptions[PC_Float] = LOCTEXT("RealCategory_DisplayAsFloat_SinglePrecision", "Float (single-precision)");
+			CategoryDescriptions[PC_Double] = LOCTEXT("RealCategory_DisplayAsFloat_DoublePrecision", "Float (double-precision)");
 			break;
 		case EBlueprintRealDisplayMode::Number:
 			CategoryDescriptions[PC_Real] = LOCTEXT("RealCategory_DisplayAsNumber", "Number");
-			CategoryDescriptions[PC_Float] = LOCTEXT("RealCategory_DisplayAsNumber", "Number (single-precision)");
-			CategoryDescriptions[PC_Double] = LOCTEXT("RealCategory_DisplayAsNumber", "Number (double-precision)");
+			CategoryDescriptions[PC_Float] = LOCTEXT("RealCategory_DisplayAsNumber_SinglePrecision", "Number (single-precision)");
+			CategoryDescriptions[PC_Double] = LOCTEXT("RealCategory_DisplayAsNumber_DoublePrecision", "Number (double-precision)");
 			break;
 		default:
 			check(false);
diff --git a/Engine/Source/Editor/DataLayerEditor/Private/DataLayer/DataLayerMode.cpp b/Engine/Source/Editor/DataLayerEditor/Private/DataLayer/DataLayerMode.cpp
--- a/Engine/Source/Editor/DataLayerEditor/Private/DataLayer/DataLayerMode.cpp
+++ b/Engine/Source/Editor/DataLayerEditor/Private/DataLayer/DataLayerMode.cpp
@@ -1077,7 +1077,7 @@ void FDataLayerMode::RegisterContextMenu()
 						Section.AddMenuEntry("ImportDataLayers", LOCTEXT("ImportDataLayers", "Import Data Layer(s)"), FText(), FSlateIcon(),
 							FUIAction(FExecuteAction::CreateLambda([this, DataLayerAssetsToImport, CreateNewDataLayer]()
 								{
-									const FScopedTransaction Transaction(LOCTEXT("ImportDataLayers", "Import Data Layers"));
+									const FScopedTransaction Transaction(LOCTEXT("ImportDataLayersTransaction", "Import Data Layers"));
 									CreateNewDataLayer(nullptr, DataLayerAssetsToImport);
 								})));
 					}
diff --git a/Engine/Source/Editor/GameProjectGeneration/Private/SProjectDialog.cpp b/Engine/Source/Editor/GameProjectGeneration/Private/SProjectDialog.cpp
--- a/Engine/Source/Editor/GameProjectGeneration/Private/SProjectDialog.cpp
+++ b/Engine/Source/Editor/GameProjectGeneration/Private/SProjectDialog.cpp
@@ -2044,7 +2044,7 @@ void SProjectDialog::UpdateProjectFileValidity()
 					}
 					else
 					{
-						LastGlobalValidityErrorText = FText::Format(LOCTEXT("NoCompilerFoundProjectDialog", "Your IDE {0} is missing or incorrectly configured, please consider using {1}"),
+						LastGlobalValidityErrorText = FText::Format(LOCTEXT("MissingIDEProjectDialog", "Your IDE {0} is missing or incorrectly configured, please consider using {1}"),
 			FSourceCodeNavigation::GetSelectedSourceCodeIDE(), FSourceCodeNavigation::GetSuggestedSourceCodeIDE());
 					}
 				}
diff --git a/Engine/Source/Editor/LandscapeEditor/Private/LandscapeEditorDetailCustomization_MiscTools.cpp b/Engine/Source/Editor/LandscapeEditor/Private/LandscapeEditorDetailCustomization_MiscTools.cpp
--- a/Engine/Source/Editor/LandscapeEditor/Private/LandscapeEditorDetailCustomization_MiscTools.cpp
+++ b/Engine/Source/Editor/LandscapeEditor/Private/LandscapeEditorDetailCustomization_MiscTools.cpp
@@ -164,7 +164,7 @@ void FLandscapeEditorDetailCustomization_MiscTools::CustomizeDetails(IDetailLayo
 			+ SHorizontalBox::Slot()
 			[
 				SNew(SButton)
-				.ToolTipText(LOCTEXT("Spline.ControlPoints.All.Tooltip", "Selects all landscape spline segments in the map."))
+				.ToolTipText(LOCTEXT("Spline.Segments.All.Tooltip", "Selects all landscape spline segments in the map."))
 				.Text(LOCTEXT("Spline.Segments", "Segments"))
 				.HAlign(HAlign_Center)
 				.OnClicked_Static(&FLandscapeEditorDetailCustomization_MiscTools::OnSelectAllSegmentsButtonClicked)
diff --git a/Engine/Source/Editor/SkeletonEditor/Private/SSkeletonTree.cpp b/Engine/Source/Editor/SkeletonEditor/Private/SSkeletonTree.cpp
--- a/Engine/Source/Editor/SkeletonEditor/Private/SSkeletonTree.cpp
+++ b/Engine/Source/Editor/SkeletonEditor/Private/SSkeletonTree.cpp
@@ -1665,7 +1665,7 @@ void SSkeletonTree::CreateBlendProfileMenu(UToolMenu* InMenu)
 			}
 
 			{
-				FToolMenuSection& NewSection = InMenu->AddSection(TEXT("BlendProfileNew"), LOCTEXT("BlendProfiles", "New"));
+				FToolMenuSection& NewSection = InMenu->AddSection(TEXT("BlendProfileNew"), LOCTEXT("NewBlendProfiles", "New"));
 				NewSection.AddMenuEntry(Actions.CreateTimeBlendProfile);
 				NewSection.AddMenuEntry(Actions.CreateWeightBlendProfile);
 				NewSection.AddMenuEntry(Actions.CreateBlendMask);
diff --git a/Engine/Source/Editor/UMGEditor/Private/AssetTypeActions_WidgetBlueprint.cpp b/Engine/Source/Editor/UMGEditor/Private/AssetTypeActions_WidgetBlueprint.cpp
--- a/Engine/Source/Editor/UMGEditor/Private/AssetTypeActions_WidgetBlueprint.cpp
+++ b/Engine/Source/Editor/UMGEditor/Private/AssetTypeActions_WidgetBlueprint.cpp
@@ -62,7 +62,7 @@ void FAssetTypeActions_WidgetBlueprint::PerformAssetDiff(UObject* Asset1, UObjec
 	if (bIsSingleAsset)
 	{
 		// identify the assumed single asset in the window's title
-		WindowTitle = FText::Format(LOCTEXT("WidgetBlueprintDiff", "{0} - Widget Blueprint Diff"), FText::FromString(NewBlueprint->GetName()));
+		WindowTitle = FText::Format(LOCTEXT("WidgetBlueprintDiffWindowTitle", "{0} - Widget Blueprint Diff"), FText::FromString(NewBlueprint->GetName()));
 	}
 
 	SBlueprintDiff::CreateDiffWindow(WindowTitle, OldBlueprint, NewBlueprint, OldRevision, NewRevision);
diff --git a/Engine/Source/Editor/UMGEditor/Private/FieldNotification/CustomizationHelper.cpp b/Engine/Source/Editor/UMGEditor/Private/FieldNotification/CustomizationHelper.cpp
--- a/Engine/Source/Editor/UMGEditor/Private/FieldNotification/CustomizationHelper.cpp
+++ b/Engine/Source/Editor/UMGEditor/Private/FieldNotification/CustomizationHelper.cpp
@@ -217,4 +217,5 @@ UK2Node_FunctionEntry* FCustomizationHelper::FindFunctionEntry(UObject* Object)
 }
 
 
-} //namespace
\ No newline at end of file
+} //namespace
+#undef LOCTEXT_NAMESPACE
diff --git a/Engine/Source/Editor/UnrealEd/Private/Fbx/FbxStaticMeshImport.cpp b/Engine/Source/Editor/UnrealEd/Private/Fbx/FbxStaticMeshImport.cpp
--- a/Engine/Source/Editor/UnrealEd/Private/Fbx/FbxStaticMeshImport.cpp
+++ b/Engine/Source/Editor/UnrealEd/Private/Fbx/FbxStaticMeshImport.cpp
@@ -1797,7 +1797,7 @@ UStaticMesh* UnFbx::FFbxImporter::ImportStaticMeshAsSingle(UObject* InParent, TA
 	BuildMeshSlowTask.MakeDialog();
 	for (int32 MeshIndex = 0; MeshIndex < MeshNodeArray.Num(); MeshIndex++)
 	{
-		BuildMeshSlowTask.EnterProgressFrame(1, FText::Format(LOCTEXT("FbxStaticMeshBuildGeometryTask", "Importing Static Mesh Geometry {0} of {1}."), FText::AsNumber(MeshIndex+1), FText::AsNumber(MeshNodeArray.Num())));
+		BuildMeshSlowTask.EnterProgressFrame(1, FText::Format(LOCTEXT("FbxStaticMeshBuildGeometryTaskProgress", "Importing Static Mesh Geometry {0} of {1}."), FText::AsNumber(MeshIndex+1), FText::AsNumber(MeshNodeArray.Num())));
 		FbxNode* Node = MeshNodeArray[MeshIndex];
 		if (Node->GetMesh())
 		{
diff --git a/Engine/Source/Editor/UnrealEd/Private/FbxMeshUtils.cpp b/Engine/Source/Editor/UnrealEd/Private/FbxMeshUtils.cpp
--- a/Engine/Source/Editor/UnrealEd/Private/FbxMeshUtils.cpp
+++ b/Engine/Source/Editor/UnrealEd/Private/FbxMeshUtils.cpp
@@ -119,7 +119,7 @@ namespace FbxMeshUtils
 							{
 								// Notification of success
 								FNotificationInfo NotificationInfo(FText::GetEmpty());
-								NotificationInfo.Text = FText::Format(NSLOCTEXT("UnrealEd", "LODImportSuccessful", "Static mesh LOD {0} imported successfully!"), FText::AsNumber(LODLevel));
+								NotificationInfo.Text = FText::Format(NSLOCTEXT("UnrealEd", "StaticMeshLODImportSuccessful", "Static mesh LOD {0} imported successfully!"), FText::AsNumber(LODLevel));
 								NotificationInfo.ExpireDuration = 5.0f;
 								FSlateNotificationManager::Get().AddNotification(NotificationInfo);
 							}
@@ -127,7 +127,7 @@ namespace FbxMeshUtils
 							{
 								// Notification of failure
 								FNotificationInfo NotificationInfo(FText::GetEmpty());
-								NotificationInfo.Text = FText::Format(NSLOCTEXT("UnrealEd", "LODImportFail", "Failed to import static mesh LOD {0}!"), FText::AsNumber(LODLevel));
+								NotificationInfo.Text = FText::Format(NSLOCTEXT("UnrealEd", "StaticMeshLODImportFail", "Failed to import static mesh LOD {0}!"), FText::AsNumber(LODLevel));
 								NotificationInfo.ExpireDuration = 5.0f;
 								FSlateNotificationManager::Get().AddNotification(NotificationInfo);
 							}
@@ -322,7 +322,7 @@ namespace FbxMeshUtils
 							{
 								// Notification of failure
 								FNotificationInfo NotificationInfo(FText::GetEmpty());
-								NotificationInfo.Text = FText::Format(NSLOCTEXT("UnrealEd", "LODImportFail", "Failed to import mesh for LOD {0}!"), FText::AsNumber(LODLevel));
+								NotificationInfo.Text = FText::Format(NSLOCTEXT("UnrealEd", "MeshLODImportFail", "Failed to import mesh for LOD {0}!"), FText::AsNumber(LODLevel));
 								NotificationInfo.ExpireDuration = 5.0f;
 								FSlateNotificationManager::Get().AddNotification(NotificationInfo);
 							}
@@ -629,7 +629,7 @@ namespace FbxMeshUtils
 					ReapplyClothing();
 					// Notification of failure
 					FNotificationInfo NotificationInfo(FText::GetEmpty());
-					NotificationInfo.Text = FText::Format(NSLOCTEXT("UnrealEd", "LODImportFail", "Failed to import mesh for LOD {0}!"), FText::AsNumber(SelectedLOD));
+					NotificationInfo.Text = FText::Format(NSLOCTEXT("UnrealEd", "MeshLODImportFail2", "Failed to import mesh for LOD {0}!"), FText::AsNumber(SelectedLOD));
 					NotificationInfo.ExpireDuration = 5.0f;
 					FSlateNotificationManager::Get().AddNotification(NotificationInfo);
 				}
diff --git a/Engine/Source/Runtime/Slate/Public/Widgets/Views/SListView.h b/Engine/Source/Runtime/Slate/Public/Widgets/Views/SListView.h
--- a/Engine/Source/Runtime/Slate/Public/Widgets/Views/SListView.h
+++ b/Engine/Source/Runtime/Slate/Public/Widgets/Views/SListView.h
@@ -1610,7 +1610,7 @@ public:
 				SNew(STableRow<ItemType>, SharedThis(this))
 				.Content()
 				[
-					SNew(STextBlock).Text(NSLOCTEXT("SListView", "BrokenUIMessage", "OnGeneratePinnedRow() not assigned."))
+					SNew(STextBlock).Text(NSLOCTEXT("SListView", "OnGeneratePinnedRowNotAssignedMessage", "OnGeneratePinnedRow() not assigned."))
 				];
 
 			return NewListItemWidget;
@@ -1638,7 +1638,7 @@ public:
 				SNew( STableRow<ItemType>, SharedThis(this) )
 				.Content()
 				[
-					SNew(STextBlock) .Text( NSLOCTEXT("SListView", "BrokenUIMessage", "OnGenerateWidget() not assigned.") )
+					SNew(STextBlock) .Text( NSLOCTEXT("SListView", "OnGenerateWidgetNotAssignedMessage", "OnGenerateWidget() not assigned.") )
 				];
 
 			return NewListItemWidget;
