Fix to task execution Check mark to use the proper Icon

diff --git a/Engine/Source/Developer/ProjectLauncher/Private/Widgets/Progress/SProjectLauncherTaskListRow.h b/Engine/Source/Developer/ProjectLauncher/Private/Widgets/Progress/SProjectLauncherTaskListRow.h
--- a/Engine/Source/Developer/ProjectLauncher/Private/Widgets/Progress/SProjectLauncherTaskListRow.h
+++ b/Engine/Source/Developer/ProjectLauncher/Private/Widgets/Progress/SProjectLauncherTaskListRow.h
@@ -109,7 +109,7 @@ private:
 
 			if (TaskPtr->GetStatus() == ELauncherTaskStatus::Completed)
 			{
-				return FAppStyle::GetBrush("Symbols.Check");
+				return FAppStyle::GetBrush("Icons.Check");
 			}
 
 			if (TaskPtr->GetStatus() == ELauncherTaskStatus::Failed)
