Add Project Launcher Profile setting for Original Release Version to support those platforms that need it (in addition to the Based On Release Version setting).  The setting will only appear for those platforms that opt-in via the platform specific version of FGenericPlatformProperties.

#rnx

diff --git a/Engine/Source/Developer/LauncherServices/Private/Launcher/LauncherWorker.cpp b/Engine/Source/Developer/LauncherServices/Private/Launcher/LauncherWorker.cpp
--- a/Engine/Source/Developer/LauncherServices/Private/Launcher/LauncherWorker.cpp
+++ b/Engine/Source/Developer/LauncherServices/Private/Launcher/LauncherWorker.cpp
@@ -707,6 +707,12 @@ FString FLauncherWorker::CreateUATCommand( const ILauncherProfileRef& InProfile,
 						UATCommand += TEXT(" -stagebasereleasepaks");
 					}
 				}
+
+				if (InProfile->GetOriginalReleaseVersionName().IsEmpty() == false)
+				{
+					UATCommand += TEXT(" -originalreleaseversion=");
+					UATCommand += InProfile->GetOriginalReleaseVersionName();
+				}
 			}
 
 			if (InProfile->IsGeneratingChunks())
diff --git a/Engine/Source/Developer/LauncherServices/Private/Profiles/LauncherProfile.h b/Engine/Source/Developer/LauncherServices/Private/Profiles/LauncherProfile.h
--- a/Engine/Source/Developer/LauncherServices/Private/Profiles/LauncherProfile.h
+++ b/Engine/Source/Developer/LauncherServices/Private/Profiles/LauncherProfile.h
@@ -51,6 +51,7 @@ enum ELauncherVersion
 	LAUNCHERSERVICES_ADDEDMAKEBINARYCONFIG = 30,
 	LAUNCHERSERVICES_ADDEDREFERENCECONTAINERS = 31,
 	LAUNCHERSERVICES_REMOVEDNUMCOOKERSTOSPAWN = 32,
+	LAUNCHERSERVICES_ADDEDORIGINALRELEASEVERSION = 33,
 	//ADD NEW STUFF HERE
 
 
@@ -532,6 +533,16 @@ public:
 		BasedOnReleaseVersionName = InBasedOnReleaseVersionName;
 	}
 
+	virtual FString GetOriginalReleaseVersionName() const override
+	{
+		return OriginalReleaseVersionName; 
+	}
+
+	virtual void SetOriginalReleaseVersionName(const FString& InOriginalReleaseVersionName) override
+	{
+		OriginalReleaseVersionName = InOriginalReleaseVersionName;
+	}
+
 	virtual FString GetReferenceContainerGlobalFileName() const override
 	{
 		return ReferenceContainerGlobalFileName;
@@ -1048,6 +1059,10 @@ public:
 			Archive << ReferenceContainerGlobalFileName;
 			Archive << ReferenceContainerCryptoKeysFileName;
 		}
+		if (Version >= LAUNCHERSERVICES_ADDEDORIGINALRELEASEVERSION)
+		{
+			Archive << OriginalReleaseVersionName;
+		}
 
 		else if(Archive.IsLoading())
 		{
@@ -1180,6 +1195,7 @@ public:
 		Writer.WriteValue("BasedOnReleaseVersionName", BasedOnReleaseVersionName);
 		Writer.WriteValue("ReferenceContainerGlobalFileName", ReferenceContainerGlobalFileName);
 		Writer.WriteValue("ReferenceContainerCryptoKeysFileName", ReferenceContainerCryptoKeysFileName);
+		Writer.WriteValue("OriginalReleaseVersionName", OriginalReleaseVersionName);
 		Writer.WriteValue("CreateDLC", CreateDLC);
 		Writer.WriteValue("DLCName", DLCName);
 		Writer.WriteValue("GenerateChunks", bGenerateChunks);
@@ -1398,6 +1414,11 @@ public:
 						Writer.WriteValue("basedonreleaseversion", GetBasedOnReleaseVersionName());
 						Writer.WriteValue("stagebasereleasepaks", ShouldStageBaseReleasePaks());
 					}
+
+					if (GetOriginalReleaseVersionName().IsEmpty() == false)
+					{
+						Writer.WriteValue("originalreleaseversion", GetOriginalReleaseVersionName());
+					}
 				}
 
 				if (IsGeneratingPatch())
@@ -1833,6 +1854,16 @@ public:
 			ReferenceContainerGlobalFileName.Empty();
 		}
 
+		if (Version >= LAUNCHERSERVICES_ADDEDORIGINALRELEASEVERSION)
+		{
+			OriginalReleaseVersionName = Object.GetStringField("OriginalReleaseVersionName");
+		}
+		else
+		{
+			OriginalReleaseVersionName.Empty();
+		}
+
+
 		CreateDLC = Object.GetBoolField("CreateDLC");
 		DLCName = Object.GetStringField("DLCName");
 		bGenerateChunks = Object.GetBoolField("GenerateChunks");
@@ -2900,6 +2931,9 @@ private:
 	// name of the release version to base this dlc / patch on
 	FString BasedOnReleaseVersionName;
 
+	// name of the original release version
+	FString OriginalReleaseVersionName;
+
 	// This build generate a patch based on some source content seealso PatchSourceContentPath
 	bool GeneratePatch;
 
diff --git a/Engine/Source/Developer/LauncherServices/Public/ILauncherProfile.h b/Engine/Source/Developer/LauncherServices/Public/ILauncherProfile.h
--- a/Engine/Source/Developer/LauncherServices/Public/ILauncherProfile.h
+++ b/Engine/Source/Developer/LauncherServices/Public/ILauncherProfile.h
@@ -1102,7 +1102,11 @@ public:
 	virtual FString GetBasedOnReleaseVersionName() const = 0;
 
 	virtual void SetBasedOnReleaseVersionName(const FString& InBasedOnReleaseVersion) = 0;
-	 
+
+	virtual FString GetOriginalReleaseVersionName() const = 0;
+
+	virtual void SetOriginalReleaseVersionName(const FString& InOriginalReleaseVersion) = 0;
+
 	/**
 	* Provides a database of compressed iostore chunks to reuse during the
 	* staging process. See IoStoreUtilities.cpp ReferenceContainerGlobalFileName.
diff --git a/Engine/Source/Developer/ProjectLauncher/Private/Widgets/Cook/SProjectLauncherCookByTheBookSettings.cpp b/Engine/Source/Developer/ProjectLauncher/Private/Widgets/Cook/SProjectLauncherCookByTheBookSettings.cpp
--- a/Engine/Source/Developer/ProjectLauncher/Private/Widgets/Cook/SProjectLauncherCookByTheBookSettings.cpp
+++ b/Engine/Source/Developer/ProjectLauncher/Private/Widgets/Cook/SProjectLauncherCookByTheBookSettings.cpp
@@ -438,6 +438,35 @@ TSharedRef<SWidget> SProjectLauncherCookByTheBookSettings::MakeComplexWidget()
 									]
 							]
 
+						+ SVerticalBox::Slot()
+							.AutoHeight()
+							.Padding(0.0f, 8.0f, 0.0f, 0.0f)
+							[
+								SNew(SProjectLauncherFormLabel)
+									.ToolTipText(LOCTEXT("OriginalReleaseVersionTextBoxToolTip", "The original release version which this DLC / Patch / Next release is based on."))
+									.LabelText(LOCTEXT("OriginalReleaseVersionTextBoxLabel", "Original release version of this application."))
+									.Visibility(this, &SProjectLauncherCookByTheBookSettings::HandleOriginalReleaseVersionVisibility)
+							]
+
+						+ SVerticalBox::Slot()
+							.AutoHeight()
+							.Padding(0.0, 4.0, 0.0, 0.0)
+							[
+								SNew(SHorizontalBox)
+
+								+ SHorizontalBox::Slot()
+									.FillWidth(1.0)
+									.Padding(0.0, 0.0, 0.0, 3.0)
+									[
+										// original release version text box
+										SNew(SEditableTextBox)
+											.ToolTipText(LOCTEXT("OriginalReleaseVersionTextBoxTooltip", "The original release version"))
+											.Text(this, &SProjectLauncherCookByTheBookSettings::HandleOriginalReleaseVersionNameTextBlockText)
+											.OnTextCommitted(this, &SProjectLauncherCookByTheBookSettings::HandleOriginalReleaseVersionNameCommitted)
+											.Visibility(this, &SProjectLauncherCookByTheBookSettings::HandleOriginalReleaseVersionVisibility)
+									]
+							]
+
 						// end create release version
 						//////////////////////////////////////////////////////////////////////////
 
@@ -1814,6 +1843,56 @@ void SProjectLauncherCookByTheBookSettings::HandleBasedOnReleaseVersionNameCommi
 }
 
 
+FText SProjectLauncherCookByTheBookSettings::HandleOriginalReleaseVersionNameTextBlockText() const
+{
+	const ILauncherProfilePtr SelectedProfile = Model->GetSelectedProfile();
+
+	FText result;
+
+	if (SelectedProfile.IsValid())
+	{
+		result = FText::FromString(SelectedProfile->GetOriginalReleaseVersionName());
+	}
+
+	return result;
+}
+
+
+void SProjectLauncherCookByTheBookSettings::HandleOriginalReleaseVersionNameCommitted(const FText& NewText, ETextCommit::Type CommitType)
+{
+	ILauncherProfilePtr SelectedProfile = Model->GetSelectedProfile();
+
+	if (SelectedProfile.IsValid())
+	{
+		SelectedProfile->SetOriginalReleaseVersionName(NewText.ToString());
+	}
+}
+
+
+EVisibility SProjectLauncherCookByTheBookSettings::HandleOriginalReleaseVersionVisibility() const
+{
+	const ILauncherProfilePtr SelectedProfile = Model->GetSelectedProfile();
+	const TArray<FString>& SelectedCookedPlatformNames = SelectedProfile->GetCookedPlatforms();
+	const TArray<ITargetPlatform*>& Platforms = GetTargetPlatformManager()->GetTargetPlatforms();
+
+	for (const FString& SelectedCookedPlatform : SelectedCookedPlatformNames)
+	{
+		for (const ITargetPlatform* Platform : Platforms)
+		{
+			if (SelectedCookedPlatform == Platform->PlatformName())
+			{
+				if (Platform->RequiresOriginalReleaseVersionForPatch())
+				{
+					return EVisibility::Visible;
+				}
+			}
+		}
+	}
+
+	return EVisibility::Collapsed;
+}
+
+
 FText SProjectLauncherCookByTheBookSettings::HandleDLCNameTextBlockText() const
 {
 	ILauncherProfilePtr SelectedProfile = Model->GetSelectedProfile();
diff --git a/Engine/Source/Developer/ProjectLauncher/Private/Widgets/Cook/SProjectLauncherCookByTheBookSettings.h b/Engine/Source/Developer/ProjectLauncher/Private/Widgets/Cook/SProjectLauncherCookByTheBookSettings.h
--- a/Engine/Source/Developer/ProjectLauncher/Private/Widgets/Cook/SProjectLauncherCookByTheBookSettings.h
+++ b/Engine/Source/Developer/ProjectLauncher/Private/Widgets/Cook/SProjectLauncherCookByTheBookSettings.h
@@ -170,6 +170,10 @@ private:
 	void HandleBasedOnReleaseVersionNameCommitted(const FText& NewText, ETextCommit::Type CommitType);
 	FText HandleBasedOnReleaseVersionNameTextBlockText() const;
 
+	void HandleOriginalReleaseVersionNameCommitted(const FText& NewText, ETextCommit::Type CommitType);
+	FText HandleOriginalReleaseVersionNameTextBlockText() const;
+	EVisibility HandleOriginalReleaseVersionVisibility() const;
+
 	//////////////////////////////////////////////////////////////////////////
 	// patch generation related functions
 	// Callback for check state changes of the 'GeneratePatch' check box.
diff --git a/Engine/Source/Developer/TargetPlatform/Public/Common/TargetPlatformBase.h b/Engine/Source/Developer/TargetPlatform/Public/Common/TargetPlatformBase.h
--- a/Engine/Source/Developer/TargetPlatform/Public/Common/TargetPlatformBase.h
+++ b/Engine/Source/Developer/TargetPlatform/Public/Common/TargetPlatformBase.h
@@ -383,6 +383,11 @@ public:
 		return TPlatformProperties::RequiresCookedData();
 	}
 
+	virtual bool RequiresOriginalReleaseVersionForPatch() const override
+	{
+		return TPlatformProperties::RequiresOriginalReleaseVersionForPatch();
+	}
+
 	virtual bool HasSecurePackageFormat() const override
 	{
 		return TPlatformProperties::HasSecurePackageFormat();
diff --git a/Engine/Source/Developer/TargetPlatform/Public/Interfaces/ITargetPlatform.h b/Engine/Source/Developer/TargetPlatform/Public/Interfaces/ITargetPlatform.h
--- a/Engine/Source/Developer/TargetPlatform/Public/Interfaces/ITargetPlatform.h
+++ b/Engine/Source/Developer/TargetPlatform/Public/Interfaces/ITargetPlatform.h
@@ -394,6 +394,14 @@ public:
 	 */
 	virtual bool RequiresCookedData() const = 0;
 
+	/**
+	 * Checks whether this platform requires the originally released version (in addition to the 
+	 * previously released version) to create a patch
+	 *
+	 * @return true if this platform requires the originally released version, false otherwise.
+	 */
+	virtual bool RequiresOriginalReleaseVersionForPatch() const = 0;
+
 	/**
 	* Checks whether this platform has a secure shippable package format, and therefore doesn't need any encryption or signing support
 	*
diff --git a/Engine/Source/Runtime/Core/Public/GenericPlatform/GenericPlatformProperties.h b/Engine/Source/Runtime/Core/Public/GenericPlatform/GenericPlatformProperties.h
--- a/Engine/Source/Runtime/Core/Public/GenericPlatform/GenericPlatformProperties.h
+++ b/Engine/Source/Runtime/Core/Public/GenericPlatform/GenericPlatformProperties.h
@@ -302,4 +302,11 @@ struct FGenericPlatformProperties
 	{
 		return nullptr;
 	}
+	 
+	// Whether the platform requires an original release version to make a patch
+	static FORCEINLINE bool RequiresOriginalReleaseVersionForPatch()
+	{
+		return false;
+	}
+
 };
