Make packaged builds and builds that deploy to device mutually exclusive in the Project Launcher's Profile Editor UI. This change makes the UI match the logic in LauncherWorker.cpp's CreateUATCommand() method, before this change users could expect a packaged build to be produced when deploying to device when one wouldn't be.

#rn Fix Editor Previously the Profile Editor would allow users to turn on packaging while deploy to device was also turned on. The logic in LauncherWorker.cpp's CreateUATCommand() method ignores the packaging setting if deploy to device is requested. With this change the Profile Editor's UI will not allow both packaging and deploying to device in the same profile.

diff --git a/Engine/Source/Developer/LauncherServices/Private/Profiles/LauncherProfile.h b/Engine/Source/Developer/LauncherServices/Private/Profiles/LauncherProfile.h
--- a/Engine/Source/Developer/LauncherServices/Private/Profiles/LauncherProfile.h
+++ b/Engine/Source/Developer/LauncherServices/Private/Profiles/LauncherProfile.h
@@ -2722,6 +2722,12 @@ protected:
 			ValidationErrors.Add(ELauncherProfileValidationErrors::CopyToDeviceRequiresCookByTheBook);
 		}
 
+		// Deploy: deployment by copying to devices requires no packaging
+		if ((DeploymentMode == ELauncherProfileDeploymentModes::CopyToDevice) && (PackagingMode != ELauncherProfilePackagingModes::DoNotPackage))
+		{
+			ValidationErrors.Add(ELauncherProfileValidationErrors::CopyToDeviceRequiresNoPackaging);
+		}
+
 		// Deploy: deployment by copying a packaged build to devices requires a package dir
 		if ((DeploymentMode == ELauncherProfileDeploymentModes::CopyRepository) && (PackageDir == TEXT("")))
 		{
diff --git a/Engine/Source/Developer/LauncherServices/Public/ILauncherProfile.h b/Engine/Source/Developer/LauncherServices/Public/ILauncherProfile.h
--- a/Engine/Source/Developer/LauncherServices/Public/ILauncherProfile.h
+++ b/Engine/Source/Developer/LauncherServices/Public/ILauncherProfile.h
@@ -234,6 +234,9 @@ namespace ELauncherProfileValidationErrors
 		/** FallbackBuild Target is required */
 		FallbackBuildTargetIsRequired,
 
+		/** Packaging and deploying are mutually exclusive */
+		CopyToDeviceRequiresNoPackaging,
+
 		Count
 	};
 }
diff --git a/Engine/Source/Developer/ProjectLauncher/Private/Widgets/Deploy/SProjectLauncherDeployPage.cpp b/Engine/Source/Developer/ProjectLauncher/Private/Widgets/Deploy/SProjectLauncherDeployPage.cpp
--- a/Engine/Source/Developer/ProjectLauncher/Private/Widgets/Deploy/SProjectLauncherDeployPage.cpp
+++ b/Engine/Source/Developer/ProjectLauncher/Private/Widgets/Deploy/SProjectLauncherDeployPage.cpp
@@ -123,6 +123,35 @@ void SProjectLauncherDeployPage::Construct(const FArguments& InArgs, const TShar
 					]
 			]
 
+		+ SVerticalBox::Slot()
+			.AutoHeight()
+			.Padding(0.0f, 8.0f, 0.0f, 0.0f)
+			[
+				SNew(SBorder)
+					.BorderImage(FAppStyle::GetBrush("ToolPanel.GroupBorder"))
+					.Padding(8.0f)
+					.Visibility(this, &SProjectLauncherDeployPage::HandleValidationErrorIconVisibility, ELauncherProfileValidationErrors::CopyToDeviceRequiresNoPackaging)
+					[
+						SNew(SHorizontalBox)
+
+						+ SHorizontalBox::Slot()
+							.AutoWidth()
+							[
+								SNew(SImage)
+									.Image(FAppStyle::GetBrush(TEXT("Icons.Error")))
+							]
+
+						+ SHorizontalBox::Slot()
+							.AutoWidth()
+							.Padding(4.0f, 0.0f)
+							.VAlign(VAlign_Center)
+							[
+								SNew(STextBlock)
+									.Text(LOCTEXT("CopyToDeviceRequiresNoPackaging", "This mode requires 'Do not package' packaging"))
+							]
+					]
+			]
+
 		+ SVerticalBox::Slot()
 			.AutoHeight()
 			.Padding(0.0f, 8.0f, 0.0f, 0.0f)
@@ -219,7 +248,8 @@ EVisibility SProjectLauncherDeployPage::HandleDeployToDeviceSettingsVisibility()
 	{
 		if (SelectedProfile->GetDeploymentMode() == ELauncherProfileDeploymentModes::CopyToDevice)
 		{
-			if (!SelectedProfile->HasValidationError(ELauncherProfileValidationErrors::CopyToDeviceRequiresCookByTheBook))
+			if (!SelectedProfile->HasValidationError(ELauncherProfileValidationErrors::CopyToDeviceRequiresCookByTheBook) &&
+				!SelectedProfile->HasValidationError(ELauncherProfileValidationErrors::CopyToDeviceRequiresNoPackaging))
 			{
 				return EVisibility::Visible;
 			}
diff --git a/Engine/Source/Developer/ProjectLauncher/Private/Widgets/Shared/SProjectLauncherValidation.h b/Engine/Source/Developer/ProjectLauncher/Private/Widgets/Shared/SProjectLauncherValidation.h
--- a/Engine/Source/Developer/ProjectLauncher/Private/Widgets/Shared/SProjectLauncherValidation.h
+++ b/Engine/Source/Developer/ProjectLauncher/Private/Widgets/Shared/SProjectLauncherValidation.h
@@ -191,6 +191,11 @@ public:
 				MakeValidationMessage(TEXT("Icons.Error"), LOCTEXT("FallbackBuildTargetIsRequired", "An explicit Default Build Target is required for the selected Variant.").ToString(), ELauncherProfileValidationErrors::FallbackBuildTargetIsRequired)
 			];
 
+		VertBox->AddSlot().AutoHeight()
+			[
+				MakeValidationMessage(TEXT("Icons.Error"), LOCTEXT("CopyToDeviceRequiresNoPackaging", "Deployment by copying to device requires 'Do not package' packaging.").ToString(), ELauncherProfileValidationErrors::CopyToDeviceRequiresNoPackaging)
+			];
+
 		check(VertBox->NumSlots() == ELauncherProfileValidationErrors::Count);
 
 		ChildSlot[VertBox.ToSharedRef()];
