QuickLaunch: When quick launch fails to start the launcher due to an invalid profile, then log all launcher profile validation errors to make it possible to diagnose the issue.

#rnx

diff --git a/Engine/Source/Developer/LauncherServices/Private/Launcher/Launcher.cpp b/Engine/Source/Developer/LauncherServices/Private/Launcher/Launcher.cpp
--- a/Engine/Source/Developer/LauncherServices/Private/Launcher/Launcher.cpp
+++ b/Engine/Source/Developer/LauncherServices/Private/Launcher/Launcher.cpp
@@ -4,6 +4,7 @@
 
 #include "HAL/RunnableThread.h"
 #include "Launcher/LauncherWorker.h"
+#include "Profiles/LauncherProfile.h"
 
 
 /* Static class member instantiations
@@ -29,6 +30,19 @@ ILauncherWorkerPtr FLauncher::Launch(const TSharedRef<ITargetDeviceProxyManager>
 			return Worker;
 		}
 	}
+	else
+	{
+		UE_LOG(LogLauncherProfile, Error, TEXT("Launcher profile '%s' for is not valid for launch."),
+			*Profile->GetName());
+		for (int32 I = 0; I < (int32)ELauncherProfileValidationErrors::Count; ++I)
+		{
+			ELauncherProfileValidationErrors::Type Error = (ELauncherProfileValidationErrors::Type)I;
+			if (Profile->HasValidationError(Error))
+			{
+				UE_LOG(LogLauncherProfile, Error, TEXT("ValidationError: %s"), *LexToStringLocalized(Error));
+			}
+		}
+	}
 
 	return nullptr;
 }
diff --git a/Engine/Source/Developer/LauncherServices/Private/Profiles/LauncherProfile.cpp b/Engine/Source/Developer/LauncherServices/Private/Profiles/LauncherProfile.cpp
new file mode 100644
--- /dev/null
+++ b/Engine/Source/Developer/LauncherServices/Private/Profiles/LauncherProfile.cpp
@@ -0,0 +1,80 @@
+// Copyright Epic Games, Inc. All Rights Reserved.
+
+#pragma once
+
+#include "LauncherProfile.h"
+
+#define LOCTEXT_NAMESPACE "SProjectLauncherValidation"
+
+FString LexToStringLocalized(ELauncherProfileValidationErrors::Type Value)
+{
+	static_assert(ELauncherProfileValidationErrors::Count == 30, "GetLocalizedValidationErrorMessage() needs to be updated to account for modified enum values.");
+	switch (Value)
+	{
+		case ELauncherProfileValidationErrors::CopyToDeviceRequiresCookByTheBook:
+			return LOCTEXT("CopyToDeviceRequiresCookByTheBookError", "Deployment by copying to device requires 'By The Book' cooking.").ToString();
+		case ELauncherProfileValidationErrors::CustomRolesNotSupportedYet:
+			return LOCTEXT("CustomRolesNotSupportedYet", "Custom launch roles are not supported yet.").ToString();
+		case ELauncherProfileValidationErrors::DeployedDeviceGroupRequired:
+			return LOCTEXT("DeployedDeviceGroupRequired", "A device group must be selected when deploying builds.").ToString();
+		case ELauncherProfileValidationErrors::InitialCultureNotAvailable:
+			return LOCTEXT("InitialCultureNotAvailable", "The Initial Culture selected for launch is not in the build.").ToString();
+		case ELauncherProfileValidationErrors::InitialMapNotAvailable:
+			return LOCTEXT("InitialMapNotAvailable", "The Initial Map selected for launch is not in the build.").ToString();
+		case ELauncherProfileValidationErrors::MalformedLaunchCommandLine:
+			return LOCTEXT("MalformedLaunchCommandLine", "The specified launch command line is not formatted correctly.").ToString();
+		case ELauncherProfileValidationErrors::NoBuildConfigurationSelected:
+			return LOCTEXT("NoBuildConfigurationSelectedError", "A Build Configuration must be selected.").ToString();
+		case ELauncherProfileValidationErrors::NoCookedCulturesSelected:
+			return LOCTEXT("NoCookedCulturesSelectedError", "At least one Culture must be selected when cooking by the book.").ToString();
+		case ELauncherProfileValidationErrors::NoLaunchRoleDeviceAssigned:
+			return LOCTEXT("NoLaunchRoleDeviceAssigned", "One or more launch roles do not have a device assigned.").ToString();
+		case ELauncherProfileValidationErrors::NoPlatformSelected:
+			return LOCTEXT("NoCookedPlatformSelectedError", "At least one Platform must be selected when cooking by the book.").ToString();
+		case ELauncherProfileValidationErrors::NoProjectSelected:
+			return LOCTEXT("NoBuildGameSelectedError", "A Project must be selected.").ToString();
+		case ELauncherProfileValidationErrors::NoPackageDirectorySpecified:
+			return LOCTEXT("NoPackageDirectorySpecified", "The deployment requires a package directory to be specified.").ToString();
+		case ELauncherProfileValidationErrors::NoPlatformSDKInstalled:
+			return LOCTEXT("NoPlatformSDKInstalled", "A required platform SDK is missing.").ToString();
+		case ELauncherProfileValidationErrors::UnversionedAndIncrimental:
+			return LOCTEXT("UnversionedAndIncrimental", "Unversioned build cannot be incremental.").ToString();
+		case ELauncherProfileValidationErrors::GeneratingPatchesCanOnlyRunFromByTheBookCookMode:
+			return LOCTEXT("GeneratingPatchesCanOnlyRunFromByTheBookCookMode", "Generating patch requires cook by the book mode.").ToString();
+		case ELauncherProfileValidationErrors::GeneratingMultiLevelPatchesRequiresGeneratePatch:
+			return LOCTEXT("GeneratingMultiLevelPatchesRequiresGeneratePatch", "Generating multilevel patch requires generating patch.").ToString();
+		case ELauncherProfileValidationErrors::StagingBaseReleasePaksWithoutABaseReleaseVersion:
+			return LOCTEXT("StagingBaseReleasePaksWithoutABaseReleaseVersion", "Staging base release pak files requires a base release version to be specified").ToString();
+		case ELauncherProfileValidationErrors::GeneratingChunksRequiresCookByTheBook:
+			return LOCTEXT("GeneratingChunksRequiresCookByTheBook", "Generating Chunks requires cook by the book mode.").ToString();
+		case ELauncherProfileValidationErrors::GeneratingChunksRequiresUnrealPak:
+			return LOCTEXT("GeneratingChunksRequiresUnrealPak", "UnrealPak must be selected to Generate Chunks.").ToString();
+		case ELauncherProfileValidationErrors::GeneratingHttpChunkDataRequiresGeneratingChunks:
+			return LOCTEXT("GeneratingHttpChunkDataRequiresGeneratingChunks", "Generate Chunks must be selected to Generate Http Chunk Install Data.").ToString();
+		case ELauncherProfileValidationErrors::GeneratingHttpChunkDataRequiresValidDirectoryAndName:
+			return LOCTEXT("GeneratingHttpChunkDataRequiresValidDirectoryAndName", "Generating Http Chunk Install Data requires a valid directory and release name.").ToString();
+		case ELauncherProfileValidationErrors::ShippingDoesntSupportCommandlineOptionsCantUseCookOnTheFly:
+			return LOCTEXT("ShippingDoesntSupportCommandlineOptionsCantUseCookOnTheFly", "Shipping doesn't support commandline options and can't use cook on the fly").ToString();
+		case ELauncherProfileValidationErrors::CookOnTheFlyDoesntSupportServer:
+			return LOCTEXT("CookOnTheFlyDoesntSupportServer", "Cook on the fly doesn't support server target configurations").ToString();
+		case ELauncherProfileValidationErrors::LaunchDeviceIsUnauthorized:
+			return LOCTEXT("LaunchDeviceIsUnauthorized", "Device is unauthorized or locked.").ToString();
+		case ELauncherProfileValidationErrors::NoArchiveDirectorySpecified:
+			return LOCTEXT("NoArchiveDirectorySpecifiedError", "The archive step requires a valid directory.").ToString();
+		case ELauncherProfileValidationErrors::IoStoreRequiresPakFiles:
+			return LOCTEXT("IoStoreRequiresPakFilesError", "UnrealPak must be selected when using I/O store container file(s)").ToString();
+		case ELauncherProfileValidationErrors::BuildTargetCookVariantMismatch:
+			return LOCTEXT("BuildTargetCookVariantMismatch", "Build Target and Cook Variant mismatch.").ToString();
+		case ELauncherProfileValidationErrors::BuildTargetIsRequired:
+			return LOCTEXT("BuildTargetIsRequired", "This profile requires an explicit Build Target set.").ToString();
+		case ELauncherProfileValidationErrors::FallbackBuildTargetIsRequired:
+			return LOCTEXT("FallbackBuildTargetIsRequired", "An explicit Default Build Target is required for the selected Variant.").ToString();
+		case ELauncherProfileValidationErrors::CopyToDeviceRequiresNoPackaging:
+			return LOCTEXT("CopyToDeviceRequiresNoPackaging", "Deployment by copying to device requires 'Do not package' packaging.").ToString();
+		default:
+			return TEXT("Unknown");
+	};
+}
+
+#undef LOCTEXT_NAMESPACE
+
diff --git a/Engine/Source/Developer/LauncherServices/Public/ILauncherProfile.h b/Engine/Source/Developer/LauncherServices/Public/ILauncherProfile.h
--- a/Engine/Source/Developer/LauncherServices/Public/ILauncherProfile.h
+++ b/Engine/Source/Developer/LauncherServices/Public/ILauncherProfile.h
@@ -241,6 +241,7 @@ namespace ELauncherProfileValidationErrors
 	};
 }
 
+LAUNCHERSERVICES_API FString LexToStringLocalized(ELauncherProfileValidationErrors::Type Value);
 
 /** Type definition for shared pointers to instances of ILauncherProfile. */
 typedef TSharedPtr<class ILauncherSimpleProfile> ILauncherSimpleProfilePtr;
diff --git a/Engine/Source/Developer/ProjectLauncher/Private/Widgets/Shared/SProjectLauncherValidation.h b/Engine/Source/Developer/ProjectLauncher/Private/Widgets/Shared/SProjectLauncherValidation.h
--- a/Engine/Source/Developer/ProjectLauncher/Private/Widgets/Shared/SProjectLauncherValidation.h
+++ b/Engine/Source/Developer/ProjectLauncher/Private/Widgets/Shared/SProjectLauncherValidation.h
@@ -48,67 +48,67 @@ public:
 
 		VertBox->AddSlot().AutoHeight()
 			[
-				MakeValidationMessage(TEXT("Icons.Error"), LOCTEXT("CopyToDeviceRequiresCookByTheBookError", "Deployment by copying to device requires 'By The Book' cooking.").ToString(), ELauncherProfileValidationErrors::CopyToDeviceRequiresCookByTheBook)
+				MakeValidationMessage(TEXT("Icons.Error"), ELauncherProfileValidationErrors::CopyToDeviceRequiresCookByTheBook)
 			];
 
 		VertBox->AddSlot().AutoHeight()
 			[
-				MakeValidationMessage(TEXT("Icons.Error"), LOCTEXT("CustomRolesNotSupportedYet", "Custom launch roles are not supported yet.").ToString(), ELauncherProfileValidationErrors::CustomRolesNotSupportedYet)
+				MakeValidationMessage(TEXT("Icons.Error"), ELauncherProfileValidationErrors::CustomRolesNotSupportedYet)
 			];
 
 		VertBox->AddSlot().AutoHeight()
 			[
-				MakeValidationMessage(TEXT("Icons.Error"), LOCTEXT("DeployedDeviceGroupRequired", "A device group must be selected when deploying builds.").ToString(), ELauncherProfileValidationErrors::DeployedDeviceGroupRequired)
+				MakeValidationMessage(TEXT("Icons.Error"), ELauncherProfileValidationErrors::DeployedDeviceGroupRequired)
 			];
 
 		VertBox->AddSlot().AutoHeight()
 			[
-				MakeValidationMessage(TEXT("Icons.Error"), LOCTEXT("InitialCultureNotAvailable", "The Initial Culture selected for launch is not in the build.").ToString(), ELauncherProfileValidationErrors::InitialCultureNotAvailable)
+				MakeValidationMessage(TEXT("Icons.Error"), ELauncherProfileValidationErrors::InitialCultureNotAvailable)
 			];
 
 		VertBox->AddSlot().AutoHeight()
 			[
-				MakeValidationMessage(TEXT("Icons.Error"), LOCTEXT("InitialMapNotAvailable", "The Initial Map selected for launch is not in the build.").ToString(), ELauncherProfileValidationErrors::InitialMapNotAvailable)
+				MakeValidationMessage(TEXT("Icons.Error"), ELauncherProfileValidationErrors::InitialMapNotAvailable)
 			];
 
 		VertBox->AddSlot().AutoHeight()
 			[
-				MakeValidationMessage(TEXT("Icons.Error"), LOCTEXT("MalformedLaunchCommandLine", "The specified launch command line is not formatted correctly.").ToString(), ELauncherProfileValidationErrors::MalformedLaunchCommandLine)
+				MakeValidationMessage(TEXT("Icons.Error"), ELauncherProfileValidationErrors::MalformedLaunchCommandLine)
 			];
 
 		VertBox->AddSlot().AutoHeight()
 			[
-				MakeValidationMessage(TEXT("Icons.Error"), LOCTEXT("NoBuildConfigurationSelectedError", "A Build Configuration must be selected.").ToString(), ELauncherProfileValidationErrors::NoBuildConfigurationSelected)
+				MakeValidationMessage(TEXT("Icons.Error"), ELauncherProfileValidationErrors::NoBuildConfigurationSelected)
 			];
 
 		VertBox->AddSlot().AutoHeight()
 			[
-				MakeValidationMessage(TEXT("Icons.Error"), LOCTEXT("NoCookedCulturesSelectedError", "At least one Culture must be selected when cooking by the book.").ToString(), ELauncherProfileValidationErrors::NoCookedCulturesSelected)
+				MakeValidationMessage(TEXT("Icons.Error"), ELauncherProfileValidationErrors::NoCookedCulturesSelected)
 			];
 
 		VertBox->AddSlot().AutoHeight()
 			[
-				MakeValidationMessage(TEXT("Icons.Error"), LOCTEXT("NoLaunchRoleDeviceAssigned", "One or more launch roles do not have a device assigned.").ToString(), ELauncherProfileValidationErrors::NoLaunchRoleDeviceAssigned)
+				MakeValidationMessage(TEXT("Icons.Error"), ELauncherProfileValidationErrors::NoLaunchRoleDeviceAssigned)
 			];
 
 		VertBox->AddSlot().AutoHeight()
 			[
-				MakeValidationMessage(TEXT("Icons.Error"), LOCTEXT("NoCookedPlatformSelectedError", "At least one Platform must be selected when cooking by the book.").ToString(), ELauncherProfileValidationErrors::NoPlatformSelected)
+				MakeValidationMessage(TEXT("Icons.Error"), ELauncherProfileValidationErrors::NoPlatformSelected)
 			];
 
 		VertBox->AddSlot().AutoHeight()
 			[
-				MakeValidationMessage(TEXT("Icons.Error"), LOCTEXT("NoBuildGameSelectedError", "A Project must be selected.").ToString(), ELauncherProfileValidationErrors::NoProjectSelected)
+				MakeValidationMessage(TEXT("Icons.Error"), ELauncherProfileValidationErrors::NoProjectSelected)
 			];
 
 		VertBox->AddSlot().AutoHeight()
 			[
-				MakeValidationMessage(TEXT("Icons.Error"), LOCTEXT("NoPackageDirectorySpecified", "The deployment requires a package directory to be specified.").ToString(), ELauncherProfileValidationErrors::NoPackageDirectorySpecified)
+				MakeValidationMessage(TEXT("Icons.Error"), ELauncherProfileValidationErrors::NoPackageDirectorySpecified)
 			];
 
 		VertBox->AddSlot().AutoHeight()
 			[
-			MakeValidationMessage(TEXT("Icons.Error"), LOCTEXT("LaunchDeviceIsUnauthorized", "Device is unauthorized or locked.").ToString(), ELauncherProfileValidationErrors::LaunchDeviceIsUnauthorized)
+				MakeValidationMessage(TEXT("Icons.Error"), ELauncherProfileValidationErrors::LaunchDeviceIsUnauthorized)
 			];
 
 		VertBox->AddSlot().AutoHeight()
@@ -118,82 +118,82 @@ public:
 
 		VertBox->AddSlot().AutoHeight()
 			[
-				MakeValidationMessage(TEXT("Icons.Error"), LOCTEXT("UnversionedAndIncrimental", "Unversioned build cannot be incremental.").ToString(), ELauncherProfileValidationErrors::UnversionedAndIncrimental)
+				MakeValidationMessage(TEXT("Icons.Error"), ELauncherProfileValidationErrors::UnversionedAndIncrimental)
 			];
 
 		VertBox->AddSlot().AutoHeight()
 			[
-				MakeValidationMessage(TEXT("Icons.Error"), LOCTEXT("GeneratingPatchesCanOnlyRunFromByTheBookCookMode", "Generating patch requires cook by the book mode.").ToString(), ELauncherProfileValidationErrors::GeneratingPatchesCanOnlyRunFromByTheBookCookMode)
+				MakeValidationMessage(TEXT("Icons.Error"), ELauncherProfileValidationErrors::GeneratingPatchesCanOnlyRunFromByTheBookCookMode)
 			];
 
 		VertBox->AddSlot().AutoHeight()
 			[
-				MakeValidationMessage(TEXT("Icons.Error"), LOCTEXT("GeneratingMultiLevelPatchesRequiresGeneratePatch", "Generating multilevel patch requires generating patch.").ToString(), ELauncherProfileValidationErrors::GeneratingMultiLevelPatchesRequiresGeneratePatch)
+				MakeValidationMessage(TEXT("Icons.Error"), ELauncherProfileValidationErrors::GeneratingMultiLevelPatchesRequiresGeneratePatch)
 			];
 
 		VertBox->AddSlot().AutoHeight()
 			[
-				MakeValidationMessage(TEXT("Icons.Error"), LOCTEXT("StagingBaseReleasePaksWithoutABaseReleaseVersion", "Staging base release pak files requires a base release version to be specified").ToString(), ELauncherProfileValidationErrors::StagingBaseReleasePaksWithoutABaseReleaseVersion)
+				MakeValidationMessage(TEXT("Icons.Error"), ELauncherProfileValidationErrors::StagingBaseReleasePaksWithoutABaseReleaseVersion)
 			];
 
 		VertBox->AddSlot().AutoHeight()
 			[
-				MakeValidationMessage(TEXT("Icons.Error"), LOCTEXT("GeneratingChunksRequiresCookByTheBook", "Generating Chunks requires cook by the book mode.").ToString(), ELauncherProfileValidationErrors::GeneratingChunksRequiresCookByTheBook)
+				MakeValidationMessage(TEXT("Icons.Error"), ELauncherProfileValidationErrors::GeneratingChunksRequiresCookByTheBook)
 			];
 
 		VertBox->AddSlot().AutoHeight()
 			[
-				MakeValidationMessage(TEXT("Icons.Error"), LOCTEXT("GeneratingChunksRequiresUnrealPak", "UnrealPak must be selected to Generate Chunks.").ToString(), ELauncherProfileValidationErrors::GeneratingChunksRequiresUnrealPak)
+				MakeValidationMessage(TEXT("Icons.Error"), ELauncherProfileValidationErrors::GeneratingChunksRequiresUnrealPak)
 			];
 
 		VertBox->AddSlot().AutoHeight()
 			[
-				MakeValidationMessage(TEXT("Icons.Error"), LOCTEXT("GeneratingHttpChunkDataRequiresGeneratingChunks", "Generate Chunks must be selected to Generate Http Chunk Install Data.").ToString(), ELauncherProfileValidationErrors::GeneratingHttpChunkDataRequiresGeneratingChunks)
+				MakeValidationMessage(TEXT("Icons.Error"), ELauncherProfileValidationErrors::GeneratingHttpChunkDataRequiresGeneratingChunks)
 			];
 
 		VertBox->AddSlot().AutoHeight()
 			[
-				MakeValidationMessage(TEXT("Icons.Error"), LOCTEXT("GeneratingHttpChunkDataRequiresValidDirectoryAndName", "Generating Http Chunk Install Data requires a valid directory and release name.").ToString(), ELauncherProfileValidationErrors::GeneratingHttpChunkDataRequiresValidDirectoryAndName)
+				MakeValidationMessage(TEXT("Icons.Error"), ELauncherProfileValidationErrors::GeneratingHttpChunkDataRequiresValidDirectoryAndName)
 			];
 
 		VertBox->AddSlot().AutoHeight()
 			[
-				MakeValidationMessage(TEXT("Icons.Error"), LOCTEXT("ShippingDoesntSupportCommandlineOptionsCantUseCookOnTheFly", "Shipping doesn't support commandline options and can't use cook on the fly").ToString(), ELauncherProfileValidationErrors::ShippingDoesntSupportCommandlineOptionsCantUseCookOnTheFly)
+				MakeValidationMessage(TEXT("Icons.Error"), ELauncherProfileValidationErrors::ShippingDoesntSupportCommandlineOptionsCantUseCookOnTheFly)
 			];
 
 		VertBox->AddSlot().AutoHeight()
 			[
-				MakeValidationMessage(TEXT("Icons.Error"), LOCTEXT("CookOnTheFlyDoesntSupportServer", "Cook on the fly doesn't support server target configurations").ToString(), ELauncherProfileValidationErrors::CookOnTheFlyDoesntSupportServer)
+				MakeValidationMessage(TEXT("Icons.Error"), ELauncherProfileValidationErrors::CookOnTheFlyDoesntSupportServer)
 			];
 
 		VertBox->AddSlot().AutoHeight()
 			[
-				MakeValidationMessage(TEXT("Icons.Error"), LOCTEXT("NoArchiveDirectorySpecifiedError", "The archive step requires a valid directory.").ToString(), ELauncherProfileValidationErrors::NoArchiveDirectorySpecified)
+				MakeValidationMessage(TEXT("Icons.Error"), ELauncherProfileValidationErrors::NoArchiveDirectorySpecified)
 			];
 
 		VertBox->AddSlot().AutoHeight()
 			[
-				MakeValidationMessage(TEXT("Icons.Error"), LOCTEXT("IoStoreRequiresPakFilesError", "UnrealPak must be selected when using I/O store container file(s)").ToString(), ELauncherProfileValidationErrors::IoStoreRequiresPakFiles)
+				MakeValidationMessage(TEXT("Icons.Error"), ELauncherProfileValidationErrors::IoStoreRequiresPakFiles)
 			];
 
 		VertBox->AddSlot().AutoHeight()
 			[
-				MakeValidationMessage(TEXT("Icons.Error"), LOCTEXT("BuildTargetCookVariantMismatch", "Build Target and Cook Variant mismatch.").ToString(), ELauncherProfileValidationErrors::BuildTargetCookVariantMismatch)
+				MakeValidationMessage(TEXT("Icons.Error"), ELauncherProfileValidationErrors::BuildTargetCookVariantMismatch)
 			];
 
 		VertBox->AddSlot().AutoHeight()
 			[
-				MakeValidationMessage(TEXT("Icons.Error"), LOCTEXT("BuildTargetIsRequired", "This profile requires an explicit Build Target set.").ToString(), ELauncherProfileValidationErrors::BuildTargetIsRequired)
+				MakeValidationMessage(TEXT("Icons.Error"), ELauncherProfileValidationErrors::BuildTargetIsRequired)
 			];
 
 		VertBox->AddSlot().AutoHeight()
 			[
-				MakeValidationMessage(TEXT("Icons.Error"), LOCTEXT("FallbackBuildTargetIsRequired", "An explicit Default Build Target is required for the selected Variant.").ToString(), ELauncherProfileValidationErrors::FallbackBuildTargetIsRequired)
+				MakeValidationMessage(TEXT("Icons.Error"), ELauncherProfileValidationErrors::FallbackBuildTargetIsRequired)
 			];
 
 		VertBox->AddSlot().AutoHeight()
 			[
-				MakeValidationMessage(TEXT("Icons.Error"), LOCTEXT("CopyToDeviceRequiresNoPackaging", "Deployment by copying to device requires 'Do not package' packaging.").ToString(), ELauncherProfileValidationErrors::CopyToDeviceRequiresNoPackaging)
+				MakeValidationMessage(TEXT("Icons.Error"), ELauncherProfileValidationErrors::CopyToDeviceRequiresNoPackaging)
 			];
 
 		check(VertBox->NumSlots() == ELauncherProfileValidationErrors::Count);
@@ -210,7 +210,7 @@ protected:
 	 * @param MessageText The message text.
 	 * @param MessageType The message type.
 	 */
-	TSharedRef<SWidget> MakeValidationMessage( const TCHAR* IconName, const FString& MessageText, ELauncherProfileValidationErrors::Type Message )
+	TSharedRef<SWidget> MakeValidationMessage( const TCHAR* IconName, ELauncherProfileValidationErrors::Type Message )
 	{
 		return SNew(SHorizontalBox)
 			.Visibility(this, &SProjectLauncherValidation::HandleValidationMessageVisibility, Message)
@@ -228,7 +228,7 @@ protected:
 			.VAlign(VAlign_Center)
 			[
 				SNew(STextBlock)
-					.Text(FText::FromString(MessageText))
+					.Text(FText::FromString(LexToStringLocalized(Message)))
 			];
 	}
 
