Slate UI for ProjectlauncherCookByTheBookSettings: Remove accidentally duplicated UI element, and remove commented-out handling for a feature that we have since abandoned.
#rnx

diff --git a/Engine/Source/Developer/ProjectLauncher/Private/Widgets/Cook/SProjectLauncherCookByTheBookSettings.cpp b/Engine/Source/Developer/ProjectLauncher/Private/Widgets/Cook/SProjectLauncherCookByTheBookSettings.cpp
--- a/Engine/Source/Developer/ProjectLauncher/Private/Widgets/Cook/SProjectLauncherCookByTheBookSettings.cpp
+++ b/Engine/Source/Developer/ProjectLauncher/Private/Widgets/Cook/SProjectLauncherCookByTheBookSettings.cpp
@@ -372,7 +372,7 @@ TSharedRef<SWidget> SProjectLauncherCookByTheBookSettings::MakeComplexWidget()
 							.AutoHeight()
 							.Padding(0.0f, 4.0f, 0.0f, 0.0f)
 							[
-								// unreal pak check box
+								// CreateReleaseVersion check box
 								SNew(SCheckBox)
 									.IsChecked(this, &SProjectLauncherCookByTheBookSettings::HandleCreateReleaseVersionCheckBoxIsChecked)
 									.OnCheckStateChanged(this, &SProjectLauncherCookByTheBookSettings::HandleCreateReleaseVersionCheckBoxCheckStateChanged)
@@ -403,7 +403,7 @@ TSharedRef<SWidget> SProjectLauncherCookByTheBookSettings::MakeComplexWidget()
 								.FillWidth(1.0)
 								.Padding(0.0, 0.0, 0.0, 3.0)
 								[
-									// repository path text box
+									// CreateReleaseVersion textbox
 									SNew(SEditableTextBox)
 										.ToolTipText(LOCTEXT("CreateReleaseVersionTextBoxTooltip", "Name of the new release to create."))
 										.Text(this, &SProjectLauncherCookByTheBookSettings::HandleCreateReleaseVersionNameTextBlockText)
@@ -430,7 +430,7 @@ TSharedRef<SWidget> SProjectLauncherCookByTheBookSettings::MakeComplexWidget()
 									.FillWidth(1.0)
 									.Padding(0.0, 0.0, 0.0, 3.0)
 									[
-										// repository path text box
+										// BasedOnReleaseVersion textbox
 										SNew(SEditableTextBox)
 											.ToolTipText(LOCTEXT("NextReleaseVersionTextBoxTooltip", "Release version to base the next release / DLC / patch on."))
 											.Text(this, &SProjectLauncherCookByTheBookSettings::HandleBasedOnReleaseVersionNameTextBlockText)
@@ -475,7 +475,7 @@ TSharedRef<SWidget> SProjectLauncherCookByTheBookSettings::MakeComplexWidget()
 							.AutoHeight()
 							.Padding(0.0f, 4.0f, 0.0f, 0.0f)
 							[
-								// unreal pak check box
+								// Generate patch checkbox
 								SNew(SCheckBox)
 									.IsChecked(this, &SProjectLauncherCookByTheBookSettings::HandleGeneratePatchCheckBoxIsChecked)
 									.OnCheckStateChanged(this, &SProjectLauncherCookByTheBookSettings::HandleGeneratePatchCheckBoxCheckStateChanged)
@@ -496,7 +496,7 @@ TSharedRef<SWidget> SProjectLauncherCookByTheBookSettings::MakeComplexWidget()
 							.AutoHeight()
 							.Padding(0.0f, 4.0f, 0.0f, 0.0f)
 							[
-								// unreal pak check box
+								// BuildDLC checkbox
 								SNew(SCheckBox)
 									.IsChecked(this, &SProjectLauncherCookByTheBookSettings::HandleBuildDLCCheckBoxIsChecked)
 									.OnCheckStateChanged(this, &SProjectLauncherCookByTheBookSettings::HandleBuildDLCCheckBoxCheckStateChanged)
@@ -539,7 +539,7 @@ TSharedRef<SWidget> SProjectLauncherCookByTheBookSettings::MakeComplexWidget()
 							.AutoHeight()
 							.Padding(0.0f, 4.0f, 0.0f, 0.0f)
 							[
-								// unreal pak check box
+								// Include engine content checkbox
 								SNew(SCheckBox)
 									.IsChecked(this, &SProjectLauncherCookByTheBookSettings::HandleDLCIncludeEngineContentCheckBoxIsChecked)
 									.OnCheckStateChanged(this, &SProjectLauncherCookByTheBookSettings::HandleDLCIncludeEngineContentCheckBoxCheckStateChanged)
@@ -572,7 +572,7 @@ TSharedRef<SWidget> SProjectLauncherCookByTheBookSettings::MakeComplexWidget()
 					+ SVerticalBox::Slot()
 						.AutoHeight()
 						[
-							// incremental cook check box
+							// incremental cook checkbox
 							SNew(SCheckBox)
 								.IsChecked(this, &SProjectLauncherCookByTheBookSettings::HandleIncrementalCheckBoxIsChecked)
 								.OnCheckStateChanged(this, &SProjectLauncherCookByTheBookSettings::HandleIncrementalCheckBoxCheckStateChanged)
@@ -585,46 +585,11 @@ TSharedRef<SWidget> SProjectLauncherCookByTheBookSettings::MakeComplexWidget()
 								]
 						]
 
-					// disabled for now until this system is live
-					/*+ SVerticalBox::Slot()
-						.AutoHeight()
-						[
-							// incremental cook check box
-							SNew(SCheckBox)
-							.IsChecked(this, &SProjectLauncherCookByTheBookSettings::HandleSharedCookedBuildCheckBoxIsChecked)
-							.OnCheckStateChanged(this, &SProjectLauncherCookByTheBookSettings::HandleSharedCookedBuildCheckBoxCheckStateChanged)
-							.Padding(FMargin(4.0f, 0.0f))
-							.ToolTipText(LOCTEXT("SharedCookedBuildCheckBoxToolTip", "Experimental: Use a build from the network to cook from."))
-							.Content()
-							[
-								SNew(STextBlock)
-								.Text(LOCTEXT("SharedCookedBuildCheckBoxText", "Iteratively cook from a pre packaged build located on the network"))
-							]
-						]*/
-
-					+ SVerticalBox::Slot()
-						.Padding(0.0f, 4.0f, 0.0f, 0.0f)
-						.AutoHeight()
-						[
-							SNew(SHorizontalBox)
-
-							+ SHorizontalBox::Slot()
-								.FillWidth(1.0)
-								.Padding(0.0, 0.0, 0.0, 3.0)
-								[
-									// repository path text box
-									SNew(SEditableTextBox)
-										.ToolTipText(LOCTEXT("NextReleaseVersionTextBoxTooltip", "Release version to base the next release / DLC / patch on."))
-										.Text(this, &SProjectLauncherCookByTheBookSettings::HandleBasedOnReleaseVersionNameTextBlockText)
-										.OnTextCommitted(this, &SProjectLauncherCookByTheBookSettings::HandleBasedOnReleaseVersionNameCommitted)
-								]
-						]
-
 					+ SVerticalBox::Slot()
 						.AutoHeight()
 						.Padding(0.0f, 4.0f, 0.0f, 0.0f)
 						[
-							// stage base release pak files check box
+							// stage base release pak files checkbox
 							SNew(SCheckBox)
 								.IsChecked(this, &SProjectLauncherCookByTheBookSettings::HandleStageBaseReleasePaksCheckBoxIsChecked)
 								.OnCheckStateChanged(this, &SProjectLauncherCookByTheBookSettings::HandleStageBaseReleasePaksCheckBoxCheckStateChanged)
@@ -637,12 +602,11 @@ TSharedRef<SWidget> SProjectLauncherCookByTheBookSettings::MakeComplexWidget()
 								]
 						]
 
-					// generate patch params
 					+ SVerticalBox::Slot()
 						.AutoHeight()
 						.Padding(0.0f, 4.0f, 0.0f, 0.0f)
 						[
-							// unreal pak check box
+							// Compress content checkbox
 							SNew(SCheckBox)
 								.IsChecked(this, &SProjectLauncherCookByTheBookSettings::HandleCompressedCheckBoxIsChecked)
 								.OnCheckStateChanged(this, &SProjectLauncherCookByTheBookSettings::HandleCompressedCheckBoxCheckStateChanged)
@@ -655,11 +619,11 @@ TSharedRef<SWidget> SProjectLauncherCookByTheBookSettings::MakeComplexWidget()
 								]
 						]
 
-					// generate new patch level params
 					+ SVerticalBox::Slot()
 						.AutoHeight()
 						.Padding(0.0f, 4.0f, 0.0f, 0.0f)
 						[
+							// generate new patch level checkbox
 							SNew(SCheckBox)
 								.IsChecked(this, &SProjectLauncherCookByTheBookSettings::HandleAddPatchLevelCheckBoxIsChecked)
 								.OnCheckStateChanged(this, &SProjectLauncherCookByTheBookSettings::HandleAddPatchLevelCheckBoxCheckStateChanged)
@@ -672,12 +636,11 @@ TSharedRef<SWidget> SProjectLauncherCookByTheBookSettings::MakeComplexWidget()
 								]
 						]
 
-					// generate dlc options
 					+ SVerticalBox::Slot()
 						.AutoHeight()
 						.Padding(0.0f, 4.0f, 0.0f, 0.0f)
 						[
-							// incremental cook check box
+							// Unversioned cooked content checkbox
 							SNew(SCheckBox)
 								.IsChecked(this, &SProjectLauncherCookByTheBookSettings::HandleUnversionedCheckBoxIsChecked)
 								.OnCheckStateChanged(this, &SProjectLauncherCookByTheBookSettings::HandleUnversionedCheckBoxCheckStateChanged)
@@ -690,12 +653,11 @@ TSharedRef<SWidget> SProjectLauncherCookByTheBookSettings::MakeComplexWidget()
 								]
 						]
 
-					// unreal pak check box
 					+ SVerticalBox::Slot()
 						.AutoHeight()
 						.Padding(0.0f, 4.0f, 0.0f, 0.0f)
 						[
-							// unreal pak check box
+							// Create .pak files checkbox
 							SNew(SCheckBox)
 								.IsChecked(this, &SProjectLauncherCookByTheBookSettings::HandleUnrealPakCheckBoxIsChecked)
 								.OnCheckStateChanged(this, &SProjectLauncherCookByTheBookSettings::HandleUnrealPakCheckBoxCheckStateChanged)
@@ -712,6 +674,7 @@ TSharedRef<SWidget> SProjectLauncherCookByTheBookSettings::MakeComplexWidget()
 						.Padding(0.0f, 4.0f, 0.0f, 0.0f)
 						.AutoHeight()
 						[
+							// Encrypt Ini Files checkbox
 							SNew(SCheckBox)
 								.IsChecked(this, &SProjectLauncherCookByTheBookSettings::HandleEncryptIniFilesCheckBoxIsChecked)
 								.OnCheckStateChanged(this, &SProjectLauncherCookByTheBookSettings::HandleEncryptIniFilesCheckBoxCheckStateChanged)
@@ -724,12 +687,11 @@ TSharedRef<SWidget> SProjectLauncherCookByTheBookSettings::MakeComplexWidget()
 								]
 						]
 
-					// generate chunks check box
 					+ SVerticalBox::Slot()
 						.AutoHeight()
 						.Padding(0.0f, 4.0f, 0.0f, 0.0f)
 						[
-						
+							// generate chunks checkbox
 							SNew(SCheckBox)
 								.IsChecked(this, &SProjectLauncherCookByTheBookSettings::HandleGenerateChunksCheckBoxIsChecked)
 								.OnCheckStateChanged(this, &SProjectLauncherCookByTheBookSettings::HandleGenerateChunksCheckBoxCheckStateChanged)
@@ -742,12 +704,11 @@ TSharedRef<SWidget> SProjectLauncherCookByTheBookSettings::MakeComplexWidget()
 								]
 						]
 
-					// don't include editor content
 					+ SVerticalBox::Slot()
 						.AutoHeight()
 						.Padding(0.0f, 4.0f, 0.0f, 0.0f)
 						[
-
+							// don't include editor content checkbox
 							SNew(SCheckBox)
 								.IsChecked(this, &SProjectLauncherCookByTheBookSettings::HandleDontIncludeEditorContentCheckBoxIsChecked)
 								.OnCheckStateChanged(this, &SProjectLauncherCookByTheBookSettings::HandleDontIncludeEditorContentCheckBoxCheckStateChanged)
@@ -776,7 +737,7 @@ TSharedRef<SWidget> SProjectLauncherCookByTheBookSettings::MakeComplexWidget()
 										.AutoHeight()
 										.Padding(0.0f, 4.0f, 0.0f, 0.0f)
 										[
-											// unreal pak check box
+											// HTTPChunkInstalls checkbox
 											SNew(SCheckBox)
 												.IsChecked(this, &SProjectLauncherCookByTheBookSettings::HandleHttpChunkInstallCheckBoxIsChecked)
 												.OnCheckStateChanged(this, &SProjectLauncherCookByTheBookSettings::HandleHttpChunkInstallCheckBoxCheckStateChanged)
@@ -1119,23 +1080,6 @@ TSharedRef<SWidget> SProjectLauncherCookByTheBookSettings::MakeSimpleWidget()
 								]
 						]
 
-					// disabled for now until this system is live
-					/*+ SVerticalBox::Slot()
-						.AutoHeight()
-						[
-							// incremental cook check box
-							SNew(SCheckBox)
-							.IsChecked(this, &SProjectLauncherCookByTheBookSettings::HandleSharedCookedBuildCheckBoxIsChecked)
-							.OnCheckStateChanged(this, &SProjectLauncherCookByTheBookSettings::HandleSharedCookedBuildCheckBoxCheckStateChanged)
-							.Padding(FMargin(4.0f, 0.0f))
-							.ToolTipText(LOCTEXT("SharedCookedBuildCheckBoxToolTip", "Experimental: Use a build from the network to cook from."))
-							.Content()
-							[
-								SNew(STextBlock)
-								.Text(LOCTEXT("SharedCookedBuildCheckBoxText", "Iteratively cook from a pre packaged build located on the network"))
-							]
-						]*/
-
 					+ SVerticalBox::Slot()
 						.AutoHeight()
 						.Padding(0.0f, 4.0f, 0.0f, 0.0f)
