Unifying naming of the reference chunk database.

diff --git a/Engine/Source/Developer/LauncherServices/Private/Profiles/LauncherProfile.h b/Engine/Source/Developer/LauncherServices/Private/Profiles/LauncherProfile.h
--- a/Engine/Source/Developer/LauncherServices/Private/Profiles/LauncherProfile.h
+++ b/Engine/Source/Developer/LauncherServices/Private/Profiles/LauncherProfile.h
@@ -1574,7 +1574,7 @@ public:
 		if (bIsStaging &&
 			GetReferenceContainerGlobalFileName().Len())
 		{
-			// (only iostore uses this) - pass reference block database
+			// (only iostore uses this) - pass reference chunk database
 			Writer.WriteValue("ReferenceContainerGlobalFileName", GetReferenceContainerGlobalFileName());
 			if (GetReferenceContainerCryptoKeysFileName().Len())
 			{
diff --git a/Engine/Source/Developer/ProjectLauncher/Private/Widgets/Package/SProjectLauncherPackagingSettings.cpp b/Engine/Source/Developer/ProjectLauncher/Private/Widgets/Package/SProjectLauncherPackagingSettings.cpp
--- a/Engine/Source/Developer/ProjectLauncher/Private/Widgets/Package/SProjectLauncherPackagingSettings.cpp
+++ b/Engine/Source/Developer/ProjectLauncher/Private/Widgets/Package/SProjectLauncherPackagingSettings.cpp
@@ -154,7 +154,7 @@ void SProjectLauncherPackagingSettings::Construct(const FArguments& InArgs, cons
 								.IsChecked(this, &SProjectLauncherPackagingSettings::HandleRetainStagedDirectoryCheckBoxIsChecked)
 								.OnCheckStateChanged(this, &SProjectLauncherPackagingSettings::HandleRetainStagedDirectoryCheckStateChanged)
 								.Padding(FMargin(4.0f, 0.0f))
-								.ToolTipText(LOCTEXT("RetainStagedDirectoryCheckBoxTooltip", "Some platforms modify I/O store containers before deployment which prevents using them as reference block database. Set this to true to save the staged directory beforehand."))
+								.ToolTipText(LOCTEXT("RetainStagedDirectoryCheckBoxTooltip", "Some platforms modify I/O store containers before deployment which prevents using them as reference chunk database. Set this to true to save the staged directory beforehand."))
 								.Content()
 								[
 									SNew(STextBlock)
@@ -186,7 +186,7 @@ void SProjectLauncherPackagingSettings::Construct(const FArguments& InArgs, cons
 							.Padding(0.0, 4.0, 0.0, 0.0)
 							[
 								SNew(STextBlock)
-								.Text(LOCTEXT("RefBlockDbFileNameText", "Optional I/O Store Reference Block Database:"))
+								.Text(LOCTEXT("RefBlockDbFileNameText", "Optional I/O Store Reference Chunk Database:"))
 							]
 
 						+ SVerticalBox::Slot()
@@ -201,9 +201,9 @@ void SProjectLauncherPackagingSettings::Construct(const FArguments& InArgs, cons
 									[
 										// path textbox for previously compressed containers.
 										SAssignNew(ReferenceContainerGlobalFileName, SEditableTextBox)
-										.IsEnabled(this, &SProjectLauncherPackagingSettings::IsReferenceBlockDbEditable)
-										.OnTextCommitted(this, &SProjectLauncherPackagingSettings::OnRefBlockDbFileNameTextCommitted)
-										.OnTextChanged(this, &SProjectLauncherPackagingSettings::OnRefBlockDbFileNameTextChanged)
+										.IsEnabled(this, &SProjectLauncherPackagingSettings::IsReferenceChunkDbEditable)
+										.OnTextCommitted(this, &SProjectLauncherPackagingSettings::OnRefChunkDbFileNameTextCommitted)
+										.OnTextChanged(this, &SProjectLauncherPackagingSettings::OnRefChunkDbFileNameTextChanged)
 										.ToolTipText(LOCTEXT("RefBlockDbUtocFileNameToolTip", "The path to a global.utoc in a directory of previously released iostore containers."))
 										.HintText(LOCTEXT("RefBlockDbUtocFileNameHint", "<path/to/global.utoc>"))
 									]
@@ -216,10 +216,10 @@ void SProjectLauncherPackagingSettings::Construct(const FArguments& InArgs, cons
 										// browse button for the global.utoc for a previously staged game.
 										SNew(SButton)
 										.ContentPadding(FMargin(6.0, 2.0))
-										.IsEnabled(this, &SProjectLauncherPackagingSettings::IsReferenceBlockDbEditable)
+										.IsEnabled(this, &SProjectLauncherPackagingSettings::IsReferenceChunkDbEditable)
 										.Text(LOCTEXT("BrowseButtonText", "Browse..."))
 										.ToolTipText(LOCTEXT("BrowseForRefDbUtocFilenameButtonToolTip", "Browse for the global.utoc file from a previously staged/released game."))
-										.OnClicked(this, &SProjectLauncherPackagingSettings::HandleRefBlockDbBrowseButtonClicked)
+										.OnClicked(this, &SProjectLauncherPackagingSettings::HandleRefChunkDbBrowseButtonClicked)
 									]
 							] // end reference container line
 
@@ -235,10 +235,10 @@ void SProjectLauncherPackagingSettings::Construct(const FArguments& InArgs, cons
 								[
 									// path textbox for previously compressed containers' crypto keys
 									SAssignNew(ReferenceContainerCryptoKeysFileName, SEditableTextBox)
-									.IsEnabled(this, &SProjectLauncherPackagingSettings::IsReferenceBlockDbEditable)
-									.OnTextCommitted(this, &SProjectLauncherPackagingSettings::OnRefBlockCryptoFileNameTextCommitted)
-									.OnTextChanged(this, &SProjectLauncherPackagingSettings::OnRefBlockCryptoFileNameTextChanged)
-									.ToolTipText(LOCTEXT("RefBlockDbCryptoFileNameToolTip", "The path to a crypto.json to decrypt the reference block containers, if needed."))
+									.IsEnabled(this, &SProjectLauncherPackagingSettings::IsReferenceChunkDbEditable)
+									.OnTextCommitted(this, &SProjectLauncherPackagingSettings::OnRefChunkCryptoFileNameTextCommitted)
+									.OnTextChanged(this, &SProjectLauncherPackagingSettings::OnRefChunkCryptoFileNameTextChanged)
+									.ToolTipText(LOCTEXT("RefBlockDbCryptoFileNameToolTip", "The path to a crypto.json to decrypt the reference chunk containers, if needed."))
 									.HintText(LOCTEXT("RefBlockDbCryptoFileNameHint", "<path/to/crypto.json, optional>"))
 								]
 
@@ -250,10 +250,10 @@ void SProjectLauncherPackagingSettings::Construct(const FArguments& InArgs, cons
 										// browse button for the crypto.json to decrypt the ref block containers, if needed
 										SNew(SButton)
 										.ContentPadding(FMargin(6.0, 2.0))
-										.IsEnabled(this, &SProjectLauncherPackagingSettings::IsReferenceBlockDbEditable)
+										.IsEnabled(this, &SProjectLauncherPackagingSettings::IsReferenceChunkDbEditable)
 										.Text(LOCTEXT("BrowseButtonText", "Browse..."))
 										.ToolTipText(LOCTEXT("BrowseForRefDbCryptoFilenameButtonToolTip", "Browse for the crypto.json file to decrypt the reference containers, if needed."))
-										.OnClicked(this, &SProjectLauncherPackagingSettings::HandleRefBlockCryptoBrowseButtonClicked)
+										.OnClicked(this, &SProjectLauncherPackagingSettings::HandleRefChunkCryptoBrowseButtonClicked)
 									]
 							] // end reference container crypto keys line
 					]
@@ -525,7 +525,7 @@ void SProjectLauncherPackagingSettings::OnTextCommitted(const FText& InText, ETe
 
 
 
-bool SProjectLauncherPackagingSettings::IsReferenceBlockDbEditable() const
+bool SProjectLauncherPackagingSettings::IsReferenceChunkDbEditable() const
 {
 	if (IsEditable() == false)
 	{
@@ -540,7 +540,7 @@ bool SProjectLauncherPackagingSettings::IsReferenceBlockDbEditable() const
 	return false;
 }
 
-void SProjectLauncherPackagingSettings::OnRefBlockDbFileNameTextChanged(const FText& InText)
+void SProjectLauncherPackagingSettings::OnRefChunkDbFileNameTextChanged(const FText& InText)
 {
 	ILauncherProfilePtr SelectedProfile = Model->GetSelectedProfile();
 	if (SelectedProfile.IsValid())
@@ -550,7 +550,7 @@ void SProjectLauncherPackagingSettings::OnRefBlockDbFileNameTextChanged(const FT
 }
 
 
-void SProjectLauncherPackagingSettings::OnRefBlockDbFileNameTextCommitted(const FText& InText, ETextCommit::Type CommitInfo)
+void SProjectLauncherPackagingSettings::OnRefChunkDbFileNameTextCommitted(const FText& InText, ETextCommit::Type CommitInfo)
 {
 	if (CommitInfo == ETextCommit::OnEnter)
 	{
@@ -562,7 +562,7 @@ void SProjectLauncherPackagingSettings::OnRefBlockDbFileNameTextCommitted(const
 	}
 }
 
-FReply SProjectLauncherPackagingSettings::HandleRefBlockDbBrowseButtonClicked()
+FReply SProjectLauncherPackagingSettings::HandleRefChunkDbBrowseButtonClicked()
 {
 	IDesktopPlatform* DesktopPlatform = FDesktopPlatformModule::Get();
 	if (DesktopPlatform)
@@ -597,7 +597,7 @@ FReply SProjectLauncherPackagingSettings::HandleRefBlockDbBrowseButtonClicked()
 }
 
 
-void SProjectLauncherPackagingSettings::OnRefBlockCryptoFileNameTextChanged(const FText& InText)
+void SProjectLauncherPackagingSettings::OnRefChunkCryptoFileNameTextChanged(const FText& InText)
 {
 	ILauncherProfilePtr SelectedProfile = Model->GetSelectedProfile();
 	if (SelectedProfile.IsValid())
@@ -607,7 +607,7 @@ void SProjectLauncherPackagingSettings::OnRefBlockCryptoFileNameTextChanged(cons
 }
 
 
-void SProjectLauncherPackagingSettings::OnRefBlockCryptoFileNameTextCommitted(const FText& InText, ETextCommit::Type CommitInfo)
+void SProjectLauncherPackagingSettings::OnRefChunkCryptoFileNameTextCommitted(const FText& InText, ETextCommit::Type CommitInfo)
 {
 	if (CommitInfo == ETextCommit::OnEnter)
 	{
@@ -619,7 +619,7 @@ void SProjectLauncherPackagingSettings::OnRefBlockCryptoFileNameTextCommitted(co
 	}
 }
 
-FReply SProjectLauncherPackagingSettings::HandleRefBlockCryptoBrowseButtonClicked()
+FReply SProjectLauncherPackagingSettings::HandleRefChunkCryptoBrowseButtonClicked()
 {
 	IDesktopPlatform* DesktopPlatform = FDesktopPlatformModule::Get();
 	if (DesktopPlatform)
diff --git a/Engine/Source/Developer/ProjectLauncher/Private/Widgets/Package/SProjectLauncherPackagingSettings.h b/Engine/Source/Developer/ProjectLauncher/Private/Widgets/Package/SProjectLauncherPackagingSettings.h
--- a/Engine/Source/Developer/ProjectLauncher/Private/Widgets/Package/SProjectLauncherPackagingSettings.h
+++ b/Engine/Source/Developer/ProjectLauncher/Private/Widgets/Package/SProjectLauncherPackagingSettings.h
@@ -59,14 +59,14 @@ private:
 	void HandleMakeBinaryConfigCheckStateChanged(ECheckBoxState NewState);
 	ECheckBoxState HandleMakeBinaryConfigCheckBoxIsChecked() const;
 
-	// Reference block database handlers.
-	bool IsReferenceBlockDbEditable() const;
-	void OnRefBlockDbFileNameTextCommitted(const FText& InText, ETextCommit::Type CommitInfo);
-	void OnRefBlockDbFileNameTextChanged(const FText& InText);
-	FReply HandleRefBlockDbBrowseButtonClicked();
-	void OnRefBlockCryptoFileNameTextCommitted(const FText& InText, ETextCommit::Type CommitInfo);
-	void OnRefBlockCryptoFileNameTextChanged(const FText& InText);
-	FReply HandleRefBlockCryptoBrowseButtonClicked();
+	// Reference chunk database handlers.
+	bool IsReferenceChunkDbEditable() const;
+	void OnRefChunkDbFileNameTextCommitted(const FText& InText, ETextCommit::Type CommitInfo);
+	void OnRefChunkDbFileNameTextChanged(const FText& InText);
+	FReply HandleRefChunkDbBrowseButtonClicked();
+	void OnRefChunkCryptoFileNameTextCommitted(const FText& InText, ETextCommit::Type CommitInfo);
+	void OnRefChunkCryptoFileNameTextChanged(const FText& InText);
+	FReply HandleRefChunkCryptoBrowseButtonClicked();
 
 
 	// Callback for getting the content text of the 'Directory' label.
